- Introduces UInterface for UAnimDataController and moves implementation to new AnimationDataController (editor only) module
- Export curve types entirely for use outside of Engine
- Miscellaneous small fixes from UE5 Main

diff --git a/Engine/Plugins/Experimental/AlembicImporter/Source/AlembicLibrary/Private/AbcImporter.cpp b/Engine/Plugins/Experimental/AlembicImporter/Source/AlembicLibrary/Private/AbcImporter.cpp
--- a/Engine/Plugins/Experimental/AlembicImporter/Source/AlembicLibrary/Private/AbcImporter.cpp
+++ b/Engine/Plugins/Experimental/AlembicImporter/Source/AlembicLibrary/Private/AbcImporter.cpp
@@ -52,7 +52,6 @@ THIRD_PARTY_INCLUDES_END
 #include "GeometryCacheCodecV1.h"
 #include "Subsystems/AssetEditorSubsystem.h"
 #include "Editor.h"
-#include "Animation/AnimData/AnimDataController.h"
 
 #include "UObject/MetaData.h"
 #include "UObject/Package.h"
@@ -641,12 +640,12 @@ TArray<UObject*> FAbcImporter::ImportAsSkeletalMesh(UObject* InParent, EObjectFl
 		uint32 WedgeOffset = 0;
 		uint32 VertexOffset = 0;
 
-		UAnimDataController* Controller = Sequence->GetController();
+		IAnimationDataController& Controller = Sequence->GetController();
 
-		Controller->OpenBracket(LOCTEXT("ImportAsSkeletalMesh", "Importing Alembic Animation"));
+		Controller.OpenBracket(LOCTEXT("ImportAsSkeletalMesh", "Importing Alembic Animation"));
 
-		Controller->SetPlayLength(AbcFile->GetImportLength());
-		Controller->SetFrameRate( FFrameRate(AbcFile->GetFramerate(), 1));
+		Controller.SetPlayLength(AbcFile->GetImportLength());
+		Controller.SetFrameRate( FFrameRate(AbcFile->GetFramerate(), 1));
 
 		Sequence->ImportFileFramerate = AbcFile->GetFramerate();
 		Sequence->ImportResampleFramerate = 1.0f / (float)AbcFile->GetFramerate();
@@ -747,8 +746,8 @@ TArray<UObject*> FAbcImporter::ImportAsSkeletalMesh(UObject* InParent, EObjectFl
 				const FReferenceSkeleton& RefSkeleton = SkeletalMesh->GetRefSkeleton();
 				const TArray<FMeshBoneInfo>& BonesInfo = RefSkeleton.GetRawRefBoneInfo();
 				
-				Controller->AddBoneTrack(BonesInfo[0].Name);
-				Controller->SetBoneTrackKeys(BonesInfo[0].Name, RootBoneTrack.PosKeys, RootBoneTrack.RotKeys, RootBoneTrack.ScaleKeys);
+				Controller.AddBoneTrack(BonesInfo[0].Name);
+				Controller.SetBoneTrackKeys(BonesInfo[0].Name, RootBoneTrack.PosKeys, RootBoneTrack.RotKeys, RootBoneTrack.ScaleKeys);
 			}
 
 			// Set recompute tangent flag on skeletal mesh sections
@@ -762,10 +761,10 @@ TArray<UObject*> FAbcImporter::ImportAsSkeletalMesh(UObject* InParent, EObjectFl
 
 		SkeletalMesh->MarkPackageDirty();
 
-		Controller->UpdateCurveNamesFromSkeleton(Skeleton, ERawCurveTrackTypes::RCT_Float);
-		Controller->NotifyPopulated();
+		Controller.UpdateCurveNamesFromSkeleton(Skeleton, ERawCurveTrackTypes::RCT_Float);
+		Controller.NotifyPopulated();
 
-		Controller->CloseBracket();
+		Controller.CloseBracket();
 
 		Sequence->PostEditChange();
 		Sequence->SetPreviewMesh(SkeletalMesh);
@@ -794,13 +793,13 @@ TArray<UObject*> FAbcImporter::ImportAsSkeletalMesh(UObject* InParent, EObjectFl
 	return GeneratedObjects;
 }
 
-void FAbcImporter::SetupMorphTargetCurves(USkeleton* Skeleton, FName ConstCurveName, UAnimSequence* Sequence, const TArray<float> &CurveValues, const TArray<float>& TimeValues, UAnimDataController* Controller)
+void FAbcImporter::SetupMorphTargetCurves(USkeleton* Skeleton, FName ConstCurveName, UAnimSequence* Sequence, const TArray<float> &CurveValues, const TArray<float>& TimeValues, IAnimationDataController& Controller)
 {
 	FSmartName NewName;
 	Skeleton->AddSmartNameAndModify(USkeleton::AnimCurveMappingName, ConstCurveName, NewName);
 
 	FAnimationCurveIdentifier CurveId(NewName, ERawCurveTrackTypes::RCT_Float);
-	Controller->AddCurve(CurveId);
+	Controller.AddCurve(CurveId);
 
 	const FFloatCurve* NewCurve = Sequence->GetDataModel()->FindFloatCurve(CurveId);
 	ensure(NewCurve);
@@ -822,7 +821,7 @@ void FAbcImporter::SetupMorphTargetCurves(USkeleton* Skeleton, FName ConstCurveN
 		RichCurve.SetKeyTangentWeightMode(NewKeyHandle, NewTangentWeightMode);
 	}
 
-	Controller->SetCurveKeys(CurveId, RichCurve.GetConstRefOfKeys());
+	Controller.SetCurveKeys(CurveId, RichCurve.GetConstRefOfKeys());
 }
 
 void FAbcImporter::SetMetaData(const TArray<UObject*>& Objects)
diff --git a/Engine/Plugins/Experimental/AlembicImporter/Source/AlembicLibrary/Public/AbcImporter.h b/Engine/Plugins/Experimental/AlembicImporter/Source/AlembicLibrary/Public/AbcImporter.h
--- a/Engine/Plugins/Experimental/AlembicImporter/Source/AlembicLibrary/Public/AbcImporter.h
+++ b/Engine/Plugins/Experimental/AlembicImporter/Source/AlembicLibrary/Public/AbcImporter.h
@@ -32,8 +32,6 @@ class UAbcImportSettings;
 class FSkeletalMeshImportData;
 class UAbcAssetImportData;
 
-class UAnimDataController;
-
 struct FMeshDescription;
 struct FAbcImportData;
 struct FGeometryCacheMeshData;
@@ -204,7 +202,7 @@ private:
 	void GenerateMorphTargetVertices(FAbcMeshSample* BaseSample, TArray<FMorphTargetDelta> &MorphDeltas, FAbcMeshSample* AverageSample, uint32 WedgeOffset, const TArray<int32>& RemapIndices, const TArray<int32>& UsedVertexIndicesForMorphs, const uint32 VertexOffset, const uint32 IndexOffset);
 	
 	/** Set up correct morph target weights from the PCA compressed data */
-	void SetupMorphTargetCurves(USkeleton* Skeleton, FName ConstCurveName, UAnimSequence* Sequence, const TArray<float> &CurveValues, const TArray<float>& TimeValues, UAnimDataController* Controller);
+	void SetupMorphTargetCurves(USkeleton* Skeleton, FName ConstCurveName, UAnimSequence* Sequence, const TArray<float> &CurveValues, const TArray<float>& TimeValues, IAnimationDataController& Controller);
 	
 	/** Set the Alembic archive metadata on the given objects */
 	void SetMetaData(const TArray<UObject*>& Objects);
diff --git a/Engine/Plugins/Importers/USDImporter/Source/USDUtilities/Private/USDSkeletalDataConversion.cpp b/Engine/Plugins/Importers/USDImporter/Source/USDUtilities/Private/USDSkeletalDataConversion.cpp
--- a/Engine/Plugins/Importers/USDImporter/Source/USDUtilities/Private/USDSkeletalDataConversion.cpp
+++ b/Engine/Plugins/Importers/USDImporter/Source/USDUtilities/Private/USDSkeletalDataConversion.cpp
@@ -431,14 +431,14 @@ namespace SkelDataConversionImpl
 
 
 		const UAnimDataModel* DataModel = Sequence->GetDataModel();
-		UAnimDataController* Controller = Sequence->GetController();
+		IAnimationDataController& Controller = Sequence->GetController();
 
 		FAnimationCurveIdentifier CurveId(NewName, ERawCurveTrackTypes::RCT_Float);
 		const FFloatCurve* Curve = DataModel->FindFloatCurve(CurveId);
 		if ( !Curve )
 		{
 			// If curve doesn't exist, add one
-			Controller->AddCurve(CurveId, AACF_DefaultCurve);
+			Controller.AddCurve(CurveId, AACF_DefaultCurve);
 			Curve = DataModel->FindFloatCurve(CurveId);
 		}
 		else
@@ -452,14 +452,14 @@ namespace SkelDataConversionImpl
 				);
 			}
 
-			Controller->SetCurveFlags(CurveId, Curve->GetCurveTypeFlags() | AACF_DefaultCurve);
+			Controller.SetCurveFlags(CurveId, Curve->GetCurveTypeFlags() | AACF_DefaultCurve);
 		}
 
-		Controller->UpdateCurveNamesFromSkeleton(Skeleton, ERawCurveTrackTypes::RCT_Float);
+		Controller.UpdateCurveNamesFromSkeleton(Skeleton, ERawCurveTrackTypes::RCT_Float);
 
 		if ( Curve )
 		{
-			Controller->SetCurveKeys(CurveId, SourceData.GetConstRefOfKeys());
+			Controller.SetCurveKeys(CurveId, SourceData.GetConstRefOfKeys());
 		}
 		else
 		{
@@ -1684,11 +1684,11 @@ bool UsdToUnreal::ConvertSkelAnim( const pxr::UsdSkelSkeletonQuery& InUsdSkeleto
 
 
 	const UAnimDataModel* DataModel = OutSkeletalAnimationAsset->GetDataModel();
-	UAnimDataController* Controller = OutSkeletalAnimationAsset->GetController();
+	IAnimationDataController& Controller = OutSkeletalAnimationAsset->GetController();
 
-	Controller->OpenBracket(LOCTEXT("ImportUSDAnimData_Bracket", "Importing USD Animation Data"));
+	Controller.OpenBracket(LOCTEXT("ImportUSDAnimData_Bracket", "Importing USD Animation Data"));
 
-	Controller->ResetModel();
+	Controller.ResetModel();
 
 	// Bake the animation for each frame.
 	// An alternative route would be to convert the time samples into TransformCurves, add them to UAnimSequence::RawCurveData,
@@ -1731,8 +1731,8 @@ bool UsdToUnreal::ConvertSkelAnim( const pxr::UsdSkelSkeletonQuery& InUsdSkeleto
 
 		for ( int32 BoneIndex = 0; BoneIndex < NumBones; ++BoneIndex )
 		{
-			Controller->AddBoneTrack(BoneInfo[BoneIndex].Name);
-			Controller->SetBoneTrackKeys(BoneInfo[BoneIndex].Name, JointTracks[BoneIndex].PosKeys, JointTracks[BoneIndex].RotKeys, JointTracks[BoneIndex].ScaleKeys);
+			Controller.AddBoneTrack(BoneInfo[BoneIndex].Name);
+			Controller.SetBoneTrackKeys(BoneInfo[BoneIndex].Name, JointTracks[BoneIndex].PosKeys, JointTracks[BoneIndex].RotKeys, JointTracks[BoneIndex].ScaleKeys);
 		}
 	}
 
@@ -1919,9 +1919,9 @@ bool UsdToUnreal::ConvertSkelAnim( const pxr::UsdSkelSkeletonQuery& InUsdSkeleto
 	OutSkeletalAnimationAsset->ImportFileFramerate = Stage.Get()->GetFramesPerSecond();
 	OutSkeletalAnimationAsset->ImportResampleFramerate = FramesPerSecond;
 
-	Controller->SetPlayLength(SequenceLengthSeconds);
-	Controller->SetFrameRate(FFrameRate(FramesPerSecond, 1));
-	Controller->CloseBracket();
+	Controller.SetPlayLength(SequenceLengthSeconds);
+	Controller.SetFrameRate(FFrameRate(FramesPerSecond, 1));
+	Controller.CloseBracket();
 
 	OutSkeletalAnimationAsset->PostEditChange();
 	OutSkeletalAnimationAsset->MarkPackageDirty();
diff --git a/Engine/Plugins/VirtualProduction/Takes/Source/TakeTrackRecorders/Private/MovieSceneAnimationTrackRecorder.cpp b/Engine/Plugins/VirtualProduction/Takes/Source/TakeTrackRecorders/Private/MovieSceneAnimationTrackRecorder.cpp
--- a/Engine/Plugins/VirtualProduction/Takes/Source/TakeTrackRecorders/Private/MovieSceneAnimationTrackRecorder.cpp
+++ b/Engine/Plugins/VirtualProduction/Takes/Source/TakeTrackRecorders/Private/MovieSceneAnimationTrackRecorder.cpp
@@ -15,7 +15,6 @@
 #include "Engine/Engine.h"
 #include "LevelSequence.h"
 #include "UObject/UObjectBaseUtility.h"
-#include "Animation/AnimData/AnimDataController.h"
 
 #define LOCTEXT_NAMESPACE "MovieSceneAnimationTrackRecorder"
 
@@ -367,20 +366,20 @@ bool UMovieSceneAnimationTrackRecorder::LoadRecordedFile(const FString& FileName
 
 								AnimSequence->DeleteNotifyTrackData();
 
-								UAnimDataController* Controller = AnimSequence->GetController();
+								IAnimationDataController& Controller = AnimSequence->GetController();
 								{
-									UAnimDataController::FScopedBracket ScopedBracket(Controller, LOCTEXT("LoadRecordedFile_Bracket", "Loading recorded animation file"));
+									IAnimationDataController::FScopedBracket ScopedBracket(Controller, LOCTEXT("LoadRecordedFile_Bracket", "Loading recorded animation file"));
 
-									Controller->ResetModel();
+									Controller.ResetModel();
 
 									const float FloatDenominator = 1000.0f;
 									const float Numerator = FloatDenominator / Header.IntervalTime;
-									Controller->SetFrameRate(FFrameRate(Numerator, FloatDenominator));
+									Controller.SetFrameRate(FFrameRate(Numerator, FloatDenominator));
 
 									int32 MaxNumberOfKeys = 0;
 									for (int32 TrackIndex = 0; TrackIndex < Header.AnimationTrackNames.Num(); ++TrackIndex)
 									{
-										Controller->AddBoneTrack(Header.AnimationTrackNames[TrackIndex]);
+										Controller.AddBoneTrack(Header.AnimationTrackNames[TrackIndex]);
 
 										TArray<FVector> PosKeys;
 										TArray<FQuat> RotKeys;
@@ -399,11 +398,11 @@ bool UMovieSceneAnimationTrackRecorder::LoadRecordedFile(const FString& FileName
 										MaxNumberOfKeys = FMath::Max(MaxNumberOfKeys, RotKeys.Num());
 										MaxNumberOfKeys = FMath::Max(MaxNumberOfKeys, ScaleKeys.Num());
 
-										Controller->SetBoneTrackKeys(Header.AnimationTrackNames[TrackIndex], PosKeys, RotKeys, ScaleKeys);
+										Controller.SetBoneTrackKeys(Header.AnimationTrackNames[TrackIndex], PosKeys, RotKeys, ScaleKeys);
 									}
 
-									Controller->SetPlayLength((MaxNumberOfKeys > 1) ? (MaxNumberOfKeys - 1) * Header.IntervalTime : MINIMUM_ANIMATION_LENGTH);
-									Controller->NotifyPopulated();
+									Controller.SetPlayLength((MaxNumberOfKeys > 1) ? (MaxNumberOfKeys - 1) * Header.IntervalTime : MINIMUM_ANIMATION_LENGTH);
+									Controller.NotifyPopulated();
 								}
 
 								AnimSequence->MarkPackageDirty();
diff --git a/Engine/Source/Developer/AnimationDataController/AnimationDataController.build.cs b/Engine/Source/Developer/AnimationDataController/AnimationDataController.build.cs
new file mode 100644
--- /dev/null
+++ b/Engine/Source/Developer/AnimationDataController/AnimationDataController.build.cs
@@ -0,0 +1,36 @@
+// Copyright Epic Games, Inc. All Rights Reserved.
+
+namespace UnrealBuildTool.Rules
+{
+	public class AnimationDataController : ModuleRules
+	{
+        public AnimationDataController(ReadOnlyTargetRules Target) : base(Target)
+		{
+			PublicDependencyModuleNames.AddRange(
+				new string[]
+				{
+					"Core",
+					"CoreUObject",
+					"Engine"
+                }
+             );
+
+			PrivateDependencyModuleNames.AddRange(
+				new string[]
+				{					
+					"MovieScene"
+				}
+			);
+			
+			if (Target.Type == TargetType.Editor)
+			{
+				PrivateDependencyModuleNames.AddRange(
+					new string[]
+					{													
+						"UnrealEd"
+					}
+				);
+			}
+		}
+	}
+}
diff --git a/Engine/Source/Developer/AnimationDataController/Private/AnimDataController.cpp b/Engine/Source/Developer/AnimationDataController/Private/AnimDataController.cpp
new file mode 100644
--- /dev/null
+++ b/Engine/Source/Developer/AnimationDataController/Private/AnimDataController.cpp
@@ -0,0 +1,1907 @@
+// Copyright Epic Games, Inc. All Rights Reserved.
+
+#include "AnimDataController.h"
+#include "AnimDataControllerActions.h"
+
+#include "Animation/AnimData/AnimDataModel.h"
+#include "Animation/AnimData/CurveIdentifier.h"
+
+#include "Animation/AnimSequenceBase.h"
+#include "Animation/AnimSequence.h"
+
+
+#include "Algo/Transform.h"
+#include "UObject/NameTypes.h"
+#include "Animation/AnimCurveTypes.h"
+#include "Math/UnrealMathUtility.h"
+
+#define LOCTEXT_NAMESPACE "AnimDataController"
+
+#if WITH_EDITOR
+
+namespace UE {
+namespace Anim {
+	bool CanTransactChanges()
+	{
+		return GEngine && GEngine->CanTransact() && !GIsTransacting;
+	}
+
+	struct FScopedCompoundTransaction
+	{
+		FScopedCompoundTransaction(UE::FChangeTransactor& InTransactor, const FText& InDescription) : Transactor(InTransactor), bCreated(false)
+		{
+			if (CanTransactChanges() && !Transactor.IsTransactionPending())
+			{
+				Transactor.OpenTransaction(InDescription);
+				bCreated = true;
+			}
+		}
+
+		~FScopedCompoundTransaction()
+		{
+			if (bCreated)
+			{
+				Transactor.CloseTransaction();
+			}
+		}
+
+		UE::FChangeTransactor& Transactor;
+		bool bCreated;
+	};
+}}
+
+#define CONDITIONAL_TRANSACTION(Text) \
+	TUniquePtr<UE::Anim::FScopedCompoundTransaction> Transaction; \
+	if (UE::Anim::CanTransactChanges() && bShouldTransact) \
+	{ \
+		Transaction = MakeUnique<UE::Anim::FScopedCompoundTransaction>(ChangeTransactor, Text); \
+	}
+
+#define CONDITIONAL_BRACKET(Text) \
+	TUniquePtr<IAnimationDataController::FScopedBracket> Transaction; \
+	if (UE::Anim::CanTransactChanges() && bShouldTransact) \
+	{ \
+		Transaction = MakeUnique<IAnimationDataController::FScopedBracket>(this, Text); \
+	}
+
+#define CONDITIONAL_ACTION(ActionClass, ...) \
+	if (UE::Anim::CanTransactChanges() && bShouldTransact) \
+	{ \
+		ChangeTransactor.AddTransactionChange<ActionClass>(__VA_ARGS__); \
+	}
+
+void UAnimDataController::SetModel(UAnimDataModel* InModel)
+{	
+	if (Model != nullptr)
+	{
+		Model->GetModifiedEvent().RemoveAll(this);
+	}
+
+	Model = InModel;
+	
+	ChangeTransactor.SetTransactionObject(InModel);
+}
+
+void UAnimDataController::OpenBracket(const FText& InTitle, bool bShouldTransact /*= true*/)
+{
+	ValidateModel();
+
+	if (UE::Anim::CanTransactChanges() && !ChangeTransactor.IsTransactionPending())
+	{
+		ChangeTransactor.OpenTransaction(InTitle);
+
+		CONDITIONAL_ACTION(UE::Anim::FCloseBracketAction, InTitle.ToString());
+	}
+
+	if (BracketDepth == 0)
+	{
+		FBracketPayload Payload;
+		Payload.Description = InTitle.ToString();
+
+		Model->Notify(EAnimDataModelNotifyType::BracketOpened, Payload);
+	}
+
+	++BracketDepth;
+}
+
+void UAnimDataController::CloseBracket(bool bShouldTransact /*= true*/)
+{
+	ValidateModel();
+
+	if (BracketDepth == 0)
+	{
+		ReportError(LOCTEXT("NoExistingBracketError", "Attempt to close bracket that was not previously opened"));
+		return;
+	}
+
+	--BracketDepth;
+
+	if (BracketDepth == 0)
+	{
+		if (UE::Anim::CanTransactChanges())
+		{
+			ensure(ChangeTransactor.IsTransactionPending());
+
+			CONDITIONAL_ACTION(UE::Anim::FOpenBracketAction, TEXT("Open Bracket"));
+
+			ChangeTransactor.CloseTransaction();
+		}
+		
+		Model->Notify(EAnimDataModelNotifyType::BracketClosed);
+	}
+}
+
+void UAnimDataController::SetPlayLength(float Length, bool bShouldTransact /*= true*/)
+{
+	ValidateModel();
+
+	// Calculate whether or new play length is shorter or longer than current, set-up T0; T1 accordingly
+	// Assumption is made that time is always added or removed at/from the end
+	// Added: T0 = current length, T1 = new length
+	// Removed: T0 = current length - removed length, T1 = current length
+	const float Delta = FMath::Abs(Length - Model->PlayLength);
+	const float T0 = Length > Model->PlayLength ? Model->PlayLength : Model->PlayLength - Delta;
+	const float T1 = Length > Model->PlayLength ? Length : Model->PlayLength;
+	ResizePlayLength(Length, T0, T1, bShouldTransact);	
+}
+
+void UAnimDataController::Resize(float Length, float T0, float T1, bool bShouldTransact /*= true*/)
+{
+	ValidateModel();
+	
+	const TRange<float> PlayRange(TRange<float>::BoundsType::Inclusive(0.f), TRange<float>::BoundsType::Inclusive(Model->PlayLength));
+	if (!FMath::IsNearlyZero(Length) && Length > 0.f)
+	{
+		if (Length != Model->PlayLength)
+		{
+			// Ensure that T0 is within the curent play range
+			if (PlayRange.Contains(T0))
+			{
+				// Ensure that the start and end length of either removal or insertion are valid
+				if (T0 < T1)
+				{
+					CONDITIONAL_BRACKET(LOCTEXT("ResizeModel", "Resizing Animation Data"));
+					const bool bInserted = Length > Model->PlayLength;
+					ResizePlayLength(Length, T0, T1, bShouldTransact);
+					ResizeCurves(Length, bInserted, T0, T1, bShouldTransact);
+					ResizeAttributes(Length, bInserted, T0, T1, bShouldTransact);
+				}
+				else
+				{
+					ReportErrorf(LOCTEXT("InvalidEndTimeError", "Invalid T1, smaller that T0 value: T0 {0}, T1 {1}"), FText::AsNumber(T0), FText::AsNumber(T1));
+				}
+			}
+			else
+			{
+				ReportErrorf(LOCTEXT("InvalidStartTimeError", "Invalid T0, not within existing play range: T0 {0}, Play Length {1}"), FText::AsNumber(T0), FText::AsNumber(Model->PlayLength));
+			}			
+		}
+		else
+		{
+			ReportWarningf(LOCTEXT("SamePlayLengthWarning", "New play length is same as existing one: {0} seconds"), FText::AsNumber(Length));
+		}
+	}
+	else
+	{
+		ReportErrorf(LOCTEXT("InvalidPlayLengthError", "Invalid play length value provided: {0} seconds"), FText::AsNumber(Length));
+	}
+}
+
+void UAnimDataController::SetFrameRate(FFrameRate FrameRate, bool bShouldTransact /*= true*/)
+{
+	ValidateModel();
+
+	// Disallow invalid frame-rates, or 0.0 intervals
+	const float FrameRateInterval = FrameRate.AsInterval();
+	if ( FrameRate.IsValid() && !FMath::IsNearlyZero(FrameRateInterval) && FrameRateInterval > 0.f)
+	{
+		CONDITIONAL_TRANSACTION(LOCTEXT("SetFrameRate", "Setting Frame Rate"));
+
+		CONDITIONAL_ACTION(UE::Anim::FSetFrameRateAction, Model);
+
+		FFrameRateChangedPayload Payload;
+		Payload.PreviousFrameRate = Model->FrameRate;
+
+		Model->FrameRate = FrameRate;
+		Model->NumberOfFrames = Model->FrameRate.AsFrameTime(Model->PlayLength).RoundToFrame().Value;
+		Model->NumberOfKeys = Model->NumberOfFrames + 1;
+
+		Model->Notify(EAnimDataModelNotifyType::FrameRateChanged, Payload);
+	}
+	else
+	{
+		ReportErrorf(LOCTEXT("InvalidFrameRateError", "Invalid frame rate provided: {0}"), FrameRate.ToPrettyText());
+	}
+}
+
+
+void UAnimDataController::UpdateCurveNamesFromSkeleton(const USkeleton* Skeleton, ERawCurveTrackTypes SupportedCurveType, bool bShouldTransact /*= true*/)
+{
+	ValidateModel();
+
+	if (Skeleton)
+	{
+		if (IsSupportedCurveType(SupportedCurveType))
+		{
+			CONDITIONAL_BRACKET(LOCTEXT("ValidateRawCurves", "Validating Animation Curve Names"));
+			switch (SupportedCurveType)
+			{
+			case ERawCurveTrackTypes::RCT_Float:
+			{
+				const FSmartNameMapping* NameMapping = Skeleton->GetSmartNameContainer(USkeleton::AnimCurveMappingName);
+				for (FFloatCurve& FloatCurve : Model->CurveData.FloatCurves)
+				{
+					FSmartName NewSmartName = FloatCurve.Name;
+					NameMapping->GetName(FloatCurve.Name.UID, NewSmartName.DisplayName);
+					if (NewSmartName != FloatCurve.Name)
+					{
+						const FAnimationCurveIdentifier CurrentId(FloatCurve.Name, SupportedCurveType);
+						const FAnimationCurveIdentifier NewId(NewSmartName, SupportedCurveType);
+						RenameCurve(CurrentId, NewId, bShouldTransact);
+					}
+				}
+				break;
+			}
+			case ERawCurveTrackTypes::RCT_Transform:
+			{
+				const FSmartNameMapping* NameMapping = Skeleton->GetSmartNameContainer(USkeleton::AnimTrackCurveMappingName);
+				for (FTransformCurve& TransformCurve : Model->CurveData.TransformCurves)
+				{
+					FSmartName NewSmartName = TransformCurve.Name;
+					NameMapping->GetName(TransformCurve.Name.UID, NewSmartName.DisplayName);
+					if (NewSmartName != TransformCurve.Name)
+					{
+						const FAnimationCurveIdentifier CurrentId(TransformCurve.Name, SupportedCurveType);
+						const FAnimationCurveIdentifier NewId(NewSmartName, SupportedCurveType);
+						RenameCurve(CurrentId, NewId, bShouldTransact);
+					}
+				}
+				break;
+			}
+			}
+		}
+		else
+		{
+			const FString CurveTypeAsString = GetCurveTypeValueName(SupportedCurveType);
+			ReportWarningf(LOCTEXT("InvalidCurveTypeWarning", "Invalid curve type provided: {0} ({1})"), FText::FromString(CurveTypeAsString), FText::AsNumber((int32)SupportedCurveType));
+		}
+	}
+	else
+	{
+		ReportError(LOCTEXT("InvalidSkeletonError", "Invalid USkeleton supplied"));
+	}
+}
+
+void UAnimDataController::FindOrAddCurveNamesOnSkeleton(USkeleton* Skeleton, ERawCurveTrackTypes SupportedCurveType, bool bShouldTransact)
+{
+	ValidateModel();
+	
+	if (Skeleton)
+	{
+		if (IsSupportedCurveType(SupportedCurveType))
+		{
+			CONDITIONAL_BRACKET(LOCTEXT("FindOrAddRawCurveNames", "Updating Skeleton with Animation Curve Names"));
+			switch (SupportedCurveType)
+			{
+			case ERawCurveTrackTypes::RCT_Float:
+			{
+				for (FFloatCurve& FloatCurve : Model->CurveData.FloatCurves)
+				{
+					FSmartName NewSmartName = FloatCurve.Name;
+					Skeleton->VerifySmartName(USkeleton::AnimCurveMappingName, NewSmartName);
+					if (NewSmartName != FloatCurve.Name)
+					{
+						const FAnimationCurveIdentifier CurrentId(FloatCurve.Name, SupportedCurveType);
+						const FAnimationCurveIdentifier NewId(NewSmartName, SupportedCurveType);
+						RenameCurve(CurrentId, NewId, bShouldTransact);
+					}
+				}
+				break;
+			}
+			case ERawCurveTrackTypes::RCT_Transform:
+			{
+				for (FTransformCurve& TransformCurve : Model->CurveData.TransformCurves)
+				{
+					FSmartName NewSmartName = TransformCurve.Name;
+					Skeleton->VerifySmartName(USkeleton::AnimTrackCurveMappingName, NewSmartName);
+					if (NewSmartName != TransformCurve.Name)
+					{
+						const FAnimationCurveIdentifier CurrentId(TransformCurve.Name, SupportedCurveType);
+						const FAnimationCurveIdentifier NewId(NewSmartName, SupportedCurveType);
+						RenameCurve(CurrentId, NewId, bShouldTransact);
+					}
+				}
+				break;
+			}
+			}
+		}
+		else
+		{
+			const FString CurveTypeAsString = GetCurveTypeValueName(SupportedCurveType);
+			ReportWarningf(LOCTEXT("InvalidCurveTypeWarning", "Invalid curve type provided: {0} ({1})"), FText::FromString(CurveTypeAsString), FText::AsNumber((int32)SupportedCurveType));
+		}
+	}
+	else
+	{
+		ReportError(LOCTEXT("InvalidSkeletonError", "Invalid USkeleton supplied "));
+	}
+}
+
+bool UAnimDataController::RemoveBoneTracksMissingFromSkeleton(const USkeleton* Skeleton, bool bShouldTransact /*= true*/)
+{
+	if (!CheckOuterClass(UAnimSequence::StaticClass()))
+	{
+		return false;
+	}
+
+	if (Skeleton)
+	{
+		TArray<FName> TracksToBeRemoved;
+		const FReferenceSkeleton& ReferenceSkeleton = Skeleton->GetReferenceSkeleton();
+
+		for (FBoneAnimationTrack& Track : Model->BoneAnimationTracks)
+		{
+			if (ReferenceSkeleton.IsValidIndex(Track.BoneTreeIndex))
+			{
+				const FName BoneName = ReferenceSkeleton.GetBoneName(Track.BoneTreeIndex);
+				if (BoneName != Track.Name)
+				{
+					// Rename track			
+					Track.Name = BoneName;
+				}
+			}
+			else
+			{
+				// Try find correct bone index
+				const int32 BoneIndex = ReferenceSkeleton.FindBoneIndex(Track.Name);
+
+				if (BoneIndex != INDEX_NONE)
+				{
+					// Update bone index
+					Track.BoneTreeIndex = BoneIndex;
+				}
+				else
+				{
+					// Remove track
+					TracksToBeRemoved.Add(Track.Name);
+					ReportWarningf(LOCTEXT("InvalidBoneIndexWarning", "Unable to find bone index, animation track will be removed: {0}"), FText::FromName(Track.Name));
+				}
+			}
+		}
+
+		if (TracksToBeRemoved.Num())
+		{
+			CONDITIONAL_BRACKET(LOCTEXT("RemoveBoneTracksMissingFromSkeleton", "Validating Bone Animation Track Data against Skeleton"));
+			for (const FName& TrackName : TracksToBeRemoved)
+			{
+				RemoveBoneTrack(TrackName);
+			}
+		}
+
+		return TracksToBeRemoved.Num() > 0;
+	}
+	else
+	{
+		ReportError(LOCTEXT("InvalidSkeletonError", "Invalid USkeleton supplied"));
+	}
+
+	return false;
+}
+
+void UAnimDataController::UpdateAttributesFromSkeleton(const USkeleton* Skeleton, bool bShouldTransact)
+{
+	ValidateModel();
+	
+	if (Skeleton)
+	{
+		// Verifying that bone (names) for attribute data exist on new skeleton
+		TArray<FAnimationAttributeIdentifier> ToRemoveIdentifiers;
+		TArray<TPair<FAnimationAttributeIdentifier, int32>> ToDuplicateIdentifiers;
+			    
+		for (const FAnimatedBoneAttribute& Attribute : Model->AnimatedBoneAttributes)
+		{
+			const int32 NewBoneIndex = Skeleton->GetReferenceSkeleton().FindBoneIndex(Attribute.Identifier.GetBoneName());
+
+			if (NewBoneIndex == INDEX_NONE)
+			{
+				ToRemoveIdentifiers.Add(Attribute.Identifier);
+			}
+			else if(NewBoneIndex != Attribute.Identifier.GetBoneIndex())
+			{						
+				ToDuplicateIdentifiers.Add(TPair<FAnimationAttributeIdentifier, int32>(Attribute.Identifier, NewBoneIndex));
+			}
+		}
+
+		if (ToRemoveIdentifiers.Num() || ToDuplicateIdentifiers.Num())
+		{
+			CONDITIONAL_BRACKET(LOCTEXT("VerifyAttributeBoneNames", "Remapping Animation Attribute Data"));
+			for (const FAnimationAttributeIdentifier& Identifier : ToRemoveIdentifiers)
+			{
+				RemoveAttribute(Identifier);
+			}
+			
+			for (const TPair<FAnimationAttributeIdentifier, int32>& Pair : ToDuplicateIdentifiers)
+			{
+				FAnimationAttributeIdentifier NewIdentifier = Pair.Key;
+				NewIdentifier.BoneIndex = Pair.Value;
+
+				DuplicateAttribute(Pair.Key, NewIdentifier);
+				RemoveAttribute(Pair.Key);
+			}	
+		}
+	}
+	else
+	{
+		ReportError(LOCTEXT("InvalidSkeletonError", "Invalid USkeleton supplied"));
+	}
+}
+
+void UAnimDataController::ResetModel(bool bShouldTransact /*= true*/)
+{
+	ValidateModel();
+
+	CONDITIONAL_BRACKET(LOCTEXT("ResetModel", "Clearing Animation Data"));
+
+	RemoveAllBoneTracks();
+
+	RemoveAllCurvesOfType(ERawCurveTrackTypes::RCT_Float, bShouldTransact);
+	RemoveAllCurvesOfType(ERawCurveTrackTypes::RCT_Transform, bShouldTransact);
+
+	SetPlayLength(MINIMUM_ANIMATION_LENGTH);
+	SetFrameRate(FFrameRate(30,1));
+
+	Model->Notify(EAnimDataModelNotifyType::Reset);
+}
+
+bool UAnimDataController::AddCurve(const FAnimationCurveIdentifier& CurveId, int32 CurveFlags /*= EAnimAssetCurveFlags::AACF_Editable*/, bool bShouldTransact /*= true*/)
+{
+	ValidateModel();
+	if (CurveId.InternalName.IsValid())
+	{		
+		if (IsSupportedCurveType(CurveId.CurveType))
+		{
+			if (!Model->FindCurve(CurveId))
+			{
+				CONDITIONAL_TRANSACTION(LOCTEXT("AddRawCurve", "Adding Animation Curve"));
+
+				FCurveAddedPayload Payload;
+				Payload.Identifier = CurveId;
+				
+				auto AddNewCurve = [CurveName = CurveId.InternalName, CurveFlags](auto& CurveTypeArray)
+				{
+					CurveTypeArray.Add({ CurveName, CurveFlags});
+				};
+				
+				switch (CurveId.CurveType)
+				{
+				case ERawCurveTrackTypes::RCT_Transform:
+					AddNewCurve(Model->CurveData.TransformCurves);
+					break;
+				case ERawCurveTrackTypes::RCT_Float:
+					AddNewCurve(Model->CurveData.FloatCurves);
+					break;
+				}
+
+				CONDITIONAL_ACTION(UE::Anim::FRemoveCurveAction, CurveId);
+				Model->Notify(EAnimDataModelNotifyType::CurveAdded, Payload);
+
+				return true;
+			}
+			else
+			{
+				const FString CurveTypeAsString = GetCurveTypeValueName(CurveId.CurveType);
+				ReportWarningf(LOCTEXT("ExistingCurveNameWarning", "Curve with name {0} and type {1} ({2}) already exists"), FText::FromName(CurveId.InternalName.DisplayName), FText::FromString(CurveTypeAsString), FText::AsNumber((int32)CurveId.CurveType));
+			}			
+		}
+		else 
+		{
+			const FString CurveTypeAsString = GetCurveTypeValueName(CurveId.CurveType);
+			ReportWarningf(LOCTEXT("InvalidCurveTypeWarning", "Invalid curve type provided: {0} ({1})"), FText::FromString(CurveTypeAsString), FText::AsNumber((int32)CurveId.CurveType));
+		}		
+	}
+	else
+	{
+		const FString CurveTypeAsString = GetCurveTypeValueName(CurveId.CurveType);
+		ReportWarningf(LOCTEXT("InvalidCurveIdentifierWarning", "Invalid curve identifier provided: name: {0}, UID: {1} type: {2}"), FText::FromName(CurveId.InternalName.DisplayName), FText::AsNumber(CurveId.InternalName.UID), FText::FromString(CurveTypeAsString));
+	}
+
+	return false;
+}
+
+bool UAnimDataController::DuplicateCurve(const FAnimationCurveIdentifier& CopyCurveId, const FAnimationCurveIdentifier& NewCurveId, bool bShouldTransact /*= true*/)
+{
+	ValidateModel();
+
+	ERawCurveTrackTypes SupportedCurveType = CopyCurveId.CurveType;
+
+	if (CopyCurveId.InternalName.IsValid() && NewCurveId.InternalName.IsValid())
+	{
+		if (IsSupportedCurveType(SupportedCurveType))
+		{
+			if (CopyCurveId.CurveType == NewCurveId.CurveType)
+			{
+				if (Model->FindCurve(CopyCurveId))
+				{
+					if (!Model->FindCurve(NewCurveId))
+					{
+						CONDITIONAL_TRANSACTION(LOCTEXT("CopyRawCurve", "Duplicating Animation Curve"));
+
+						auto DuplicateCurve = [NewCurveName = NewCurveId.InternalName](auto& CurveDataArray, const auto& SourceCurve)
+						{
+							auto& DuplicatedCurve = CurveDataArray.Add_GetRef( { NewCurveName, SourceCurve.GetCurveTypeFlags() });
+							DuplicatedCurve.CopyCurve(SourceCurve);
+						};
+						
+						switch (SupportedCurveType)
+						{
+						case ERawCurveTrackTypes::RCT_Transform:
+							DuplicateCurve(Model->CurveData.TransformCurves, Model->GetTransformCurve(CopyCurveId));
+							break;
+						case ERawCurveTrackTypes::RCT_Float:
+							DuplicateCurve(Model->CurveData.FloatCurves, Model->GetFloatCurve(CopyCurveId));
+							break;
+						}
+
+						FCurveAddedPayload Payload;
+						Payload.Identifier = NewCurveId;
+						Model->Notify(EAnimDataModelNotifyType::CurveAdded, Payload);
+
+						CONDITIONAL_ACTION(UE::Anim::FRemoveCurveAction, NewCurveId);
+
+						return true;
+					}
+					else
+					{
+						const FString CurveTypeAsString = GetCurveTypeValueName(NewCurveId.CurveType);
+						ReportWarningf(LOCTEXT("ExistingCurveNameWarning", "Curve with name {0} and type {1} ({2}) already exists"), FText::FromName(NewCurveId.InternalName.DisplayName), FText::FromString(CurveTypeAsString), FText::AsNumber((int32)NewCurveId.CurveType));
+					}
+				}
+				else
+				{
+					const FString CurveTypeAsString = GetCurveTypeValueName(CopyCurveId.CurveType);
+					ReportWarningf(LOCTEXT("CurveNameToDuplicateNotFoundWarning", "Could not find curve with name {0} and type {1} ({2}) for duplication"), FText::FromName(NewCurveId.InternalName.DisplayName), FText::FromString(CurveTypeAsString), FText::AsNumber((int32)NewCurveId.CurveType));
+				}
+			}
+		}
+		else
+		{
+			const FString CurveTypeAsString = GetCurveTypeValueName(SupportedCurveType);
+			ReportWarningf(LOCTEXT("InvalidCurveTypeWarning", "Invalid curve type provided: {0} ({1})"), FText::FromString(CurveTypeAsString), FText::AsNumber((int32)SupportedCurveType));
+		}
+	}
+
+	return false;
+}
+
+
+bool UAnimDataController::RemoveCurve(const FAnimationCurveIdentifier& CurveId, bool bShouldTransact /*= true*/)
+{
+	ValidateModel();
+	ERawCurveTrackTypes SupportedCurveType = CurveId.CurveType;
+
+	if (CurveId.InternalName.IsValid())
+	{
+		if (IsSupportedCurveType(CurveId.CurveType))
+		{
+			const FAnimCurveBase* Curve = Model->FindCurve(CurveId);
+			if (Curve)
+			{
+				CONDITIONAL_TRANSACTION(LOCTEXT("RemoveCurve", "Removing Animation Curve"));
+
+				switch (SupportedCurveType)
+				{
+					case ERawCurveTrackTypes::RCT_Transform:
+					{
+						const FTransformCurve& TransformCurve = Model->GetTransformCurve(CurveId);
+						CONDITIONAL_ACTION(UE::Anim::FAddTransformCurveAction, CurveId, TransformCurve.GetCurveTypeFlags(), TransformCurve);
+						Model->CurveData.TransformCurves.RemoveAll([Name = TransformCurve.Name](const FTransformCurve& ToRemoveCurve) { return ToRemoveCurve.Name == Name; });
+						break;
+					}
+					case ERawCurveTrackTypes::RCT_Float:
+					{
+						const FFloatCurve& FloatCurve = Model->GetFloatCurve(CurveId);
+						CONDITIONAL_ACTION(UE::Anim::FAddFloatCurveAction, CurveId, FloatCurve.GetCurveTypeFlags(), FloatCurve.FloatCurve.GetConstRefOfKeys(), FloatCurve.Color);
+						Model->CurveData.FloatCurves.RemoveAll([Name = FloatCurve.Name](const FFloatCurve& ToRemoveCurve) { return ToRemoveCurve.Name == Name; });
+						break;
+					}
+				}
+
+				FCurveRemovedPayload Payload;
+				Payload.Identifier = CurveId;
+				Model->Notify(EAnimDataModelNotifyType::CurveRemoved, Payload);
+
+				return true;
+			}
+			else
+			{
+				const FString CurveTypeAsString = GetCurveTypeValueName(SupportedCurveType);
+				ReportWarningf(LOCTEXT("UnableToFindCurveForRemovalWarning", "Unable to find curve: {0} of type {1}"), FText::FromName(CurveId.InternalName.DisplayName), FText::FromString(CurveTypeAsString));
+			}
+		}
+		else
+		{
+			const FString CurveTypeAsString = GetCurveTypeValueName(CurveId.CurveType);
+			ReportWarningf(LOCTEXT("InvalidCurveTypeWarning", "Invalid curve type provided: {0} ({1})"), FText::FromString(CurveTypeAsString), FText::AsNumber((int32)CurveId.CurveType));
+		}
+	}
+
+	return false;
+}
+
+void UAnimDataController::RemoveAllCurvesOfType(ERawCurveTrackTypes SupportedCurveType /*= ERawCurveTrackTypes::RCT_Float*/, bool bShouldTransact /*= true*/)
+{
+	ValidateModel();
+
+	CONDITIONAL_BRACKET(LOCTEXT("DeleteAllRawCurve", "Deleting All Animation Curve"));
+	switch (SupportedCurveType)
+	{
+	case ERawCurveTrackTypes::RCT_Transform:
+	{
+		TArray<FTransformCurve> TransformCurves = Model->CurveData.TransformCurves;
+		for (const FTransformCurve& Curve : TransformCurves)
+		{
+			RemoveCurve(FAnimationCurveIdentifier(Curve.Name, ERawCurveTrackTypes::RCT_Transform), bShouldTransact);
+		}
+		break;
+	}
+	case ERawCurveTrackTypes::RCT_Float:
+	{
+		TArray<FFloatCurve> FloatCurves = Model->CurveData.FloatCurves;
+		for (const FFloatCurve& Curve : FloatCurves)
+		{
+			RemoveCurve(FAnimationCurveIdentifier(Curve.Name, ERawCurveTrackTypes::RCT_Float), bShouldTransact);
+		}
+		break;
+	}
+	case ERawCurveTrackTypes::RCT_Vector:
+	default:
+		const FString CurveTypeAsString = GetCurveTypeValueName(SupportedCurveType);
+		ReportWarningf(LOCTEXT("InvalidCurveTypeWarning", "Invalid curve type provided: {0} ({1})"), FText::FromString(CurveTypeAsString), FText::AsNumber((int32)SupportedCurveType));
+	}
+
+}
+
+bool UAnimDataController::SetCurveFlag(const FAnimationCurveIdentifier& CurveId, EAnimAssetCurveFlags Flag, bool bState /*= true*/, bool bShouldTransact /*= true*/)
+{
+	ValidateModel();
+
+	ERawCurveTrackTypes SupportedCurveType = CurveId.CurveType;
+
+	FAnimCurveBase* Curve = nullptr;
+
+	if (SupportedCurveType == ERawCurveTrackTypes::RCT_Float)
+	{
+		Curve = Model->FindMutableFloatCurveById(CurveId);
+	}
+	else if (SupportedCurveType == ERawCurveTrackTypes::RCT_Transform)
+	{
+		Curve = Model->FindMutableTransformCurveById(CurveId);
+	}
+	
+	if (Curve)
+	{
+		CONDITIONAL_TRANSACTION(LOCTEXT("SetCurveFlag", "Setting Raw Curve Flag"));
+
+		const int32 CurrentFlags = Curve->GetCurveTypeFlags();
+
+		CONDITIONAL_ACTION(UE::Anim::FSetCurveFlagsAction, CurveId, CurrentFlags, SupportedCurveType);
+
+		FCurveFlagsChangedPayload Payload;
+		Payload.Identifier = CurveId;
+		Payload.OldFlags = Curve->GetCurveTypeFlags();
+
+		Curve->SetCurveTypeFlag(Flag, bState);
+
+		Model->Notify(EAnimDataModelNotifyType::CurveFlagsChanged, Payload);
+
+		return true;
+	}
+	else
+	{
+		const FString CurveTypeAsString = GetCurveTypeValueName(SupportedCurveType);
+		ReportWarningf(LOCTEXT("UnableToFindCurveWarning", "Unable to find curve: {0} of type {1}"), FText::FromName(CurveId.InternalName.DisplayName), FText::FromString(CurveTypeAsString));
+	}
+
+	return false;
+}
+
+bool UAnimDataController::SetCurveFlags(const FAnimationCurveIdentifier& CurveId, int32 Flags, bool bShouldTransact /*= true*/)
+{
+	ValidateModel();
+
+	FAnimCurveBase* Curve = nullptr;
+
+	ERawCurveTrackTypes SupportedCurveType = CurveId.CurveType;
+
+	if (SupportedCurveType == ERawCurveTrackTypes::RCT_Float)
+	{
+		Curve = Model->FindMutableFloatCurveById(CurveId);
+	}
+	else if (SupportedCurveType == ERawCurveTrackTypes::RCT_Transform)
+	{
+		Curve = Model->FindMutableTransformCurveById(CurveId);
+	}
+
+	if (Curve)
+	{
+		CONDITIONAL_TRANSACTION(LOCTEXT("SetCurveFlag", "Setting Raw Curve Flags"));
+
+		const int32 CurrentFlags = Curve->GetCurveTypeFlags();
+
+		CONDITIONAL_ACTION(UE::Anim::FSetCurveFlagsAction, CurveId, CurrentFlags, SupportedCurveType);
+
+		FCurveFlagsChangedPayload Payload;
+		Payload.Identifier = CurveId;
+		Payload.OldFlags = Curve->GetCurveTypeFlags();
+
+		Curve->SetCurveTypeFlags(Flags);
+
+		Model->Notify(EAnimDataModelNotifyType::CurveFlagsChanged, Payload);
+
+		return true;
+	}
+	else
+	{
+		const FString CurveTypeAsString = GetCurveTypeValueName(SupportedCurveType);
+		ReportWarningf(LOCTEXT("UnableToFindCurveForRemovalWarning", "Unable to find curve: {0} of type {1}"), FText::FromName(CurveId.InternalName.DisplayName), FText::FromString(CurveTypeAsString));
+	}
+
+	return false;
+}
+
+bool UAnimDataController::SetTransformCurveKeys(const FAnimationCurveIdentifier& CurveId, const TArray<FTransform>& TransformValues, const TArray<float>& TimeKeys, bool bShouldTransact /*= true*/)
+{
+	ValidateModel();
+
+	if (TransformValues.Num() == TimeKeys.Num())
+	{
+		FTransformCurve* Curve = Model->FindMutableTransformCurveById(CurveId);
+
+		if (Curve)
+		{
+			CONDITIONAL_BRACKET(LOCTEXT("SetTransformCurveKeys_Bracket", "Setting Transform Curve Keys"));
+			
+			struct FKeys
+			{
+				FKeys(int32 NumKeys)
+				{
+					for (int32 ChannelIndex = 0; ChannelIndex < 3; ++ChannelIndex)
+					{
+						ChannelKeys[ChannelIndex].SetNum(NumKeys);
+					}
+				}
+
+				TArray<FRichCurveKey> ChannelKeys[3];
+			};
+
+			FKeys TranslationKeys(TransformValues.Num());
+			FKeys RotationKeys(TransformValues.Num());
+			FKeys ScaleKeys(TransformValues.Num());
+
+			FKeys* SubCurveKeys[3] = { &TranslationKeys, &RotationKeys, &ScaleKeys };
+
+			// Generate the curve keys
+			for (int32 KeyIndex = 0; KeyIndex < TransformValues.Num(); ++KeyIndex)
+			{
+				const FTransform& Value = TransformValues[KeyIndex];
+				const float& Time = TimeKeys[KeyIndex];
+
+				const FVector Translation = Value.GetLocation();
+				const FVector Rotation = Value.GetRotation().Euler();
+				const FVector Scale = Value.GetScale3D();
+
+				auto SetKey = [Time, KeyIndex](FKeys& Key, const FVector& Vector)
+				{
+					for (int32 ChannelIndex = 0; ChannelIndex < 3; ++ChannelIndex)
+					{
+						Key.ChannelKeys[ChannelIndex][KeyIndex] = FRichCurveKey(Time, Vector[ChannelIndex]);
+					}
+				};
+
+				SetKey(TranslationKeys, Translation);
+				SetKey(RotationKeys, Rotation);
+				SetKey(ScaleKeys, Scale);
+			}
+			
+			for (int32 SubCurveIndex = 0; SubCurveIndex < 3; ++SubCurveIndex)
+			{
+				const ETransformCurveChannel Channel = (ETransformCurveChannel)SubCurveIndex;
+				FKeys* CurveKeys = SubCurveKeys[SubCurveIndex];
+				for (int32 ChannelIndex = 0; ChannelIndex < 3; ++ChannelIndex)
+				{
+					const EVectorCurveChannel Axis = (EVectorCurveChannel)ChannelIndex;
+					FAnimationCurveIdentifier TargetCurveIdentifier = CurveId;
+					UAnimationCurveIdentifierExtensions::GetTransformChildCurveIdentifier(TargetCurveIdentifier, Channel, Axis);
+					SetCurveKeys(TargetCurveIdentifier, CurveKeys->ChannelKeys[ChannelIndex], bShouldTransact);
+				}
+			}
+
+			return true;
+		}
+		else
+		{
+			ReportWarningf(LOCTEXT("UnableToFindTransformCurveWarning", "Unable to find transform curve: {0}"), FText::FromName(CurveId.InternalName.DisplayName));
+		}
+	}
+	else
+	{
+		// time value mismatch
+		ReportWarningf(LOCTEXT("InvalidNumberOfTimeAndKeyEntriesWarning", "Number of times and key entries do not match: number of time values {0}, number of key values {1}"), FText::AsNumber(TimeKeys.Num()), FText::AsNumber(TransformValues.Num()));
+	}
+
+	return false;	
+}
+
+
+bool UAnimDataController::SetTransformCurveKey(const FAnimationCurveIdentifier& CurveId, float Time, const FTransform& Value, bool bShouldTransact /*= true*/)
+{
+	ValidateModel();
+
+	FTransformCurve* Curve = Model->FindMutableTransformCurveById(CurveId);
+
+	if (Curve)
+	{
+		CONDITIONAL_BRACKET(LOCTEXT("AddTransformCurveKey_Bracket", "Setting Transform Curve Key"));
+		struct FKeys
+		{
+			FRichCurveKey ChannelKeys[3];
+		};
+
+		FKeys VectorKeys[3];
+		
+		// Generate the rich curve keys		
+		const FVector Translation = Value.GetLocation();
+		const FVector Rotation = Value.GetRotation().Euler();
+		const FVector Scale = Value.GetScale3D();
+
+		auto SetKey = [Time](FKeys& Key, const FVector& Vector)
+		{
+			for (int32 ChannelIndex = 0; ChannelIndex < 3; ++ChannelIndex)
+			{
+				Key.ChannelKeys[ChannelIndex] = FRichCurveKey(Time, Vector[ChannelIndex]);
+			}
+		};
+
+		SetKey(VectorKeys[0], Translation);
+		SetKey(VectorKeys[1], Rotation);
+		SetKey(VectorKeys[2], Scale);
+		
+		for (int32 SubCurveIndex = 0; SubCurveIndex < 3; ++SubCurveIndex)
+		{
+			const ETransformCurveChannel Channel = (ETransformCurveChannel)SubCurveIndex;
+			const FKeys& VectorCurveKeys = VectorKeys[SubCurveIndex];
+			for (int32 ChannelIndex = 0; ChannelIndex < 3; ++ChannelIndex)
+			{
+				const EVectorCurveChannel Axis = (EVectorCurveChannel)ChannelIndex;
+				FAnimationCurveIdentifier TargetCurveIdentifier = CurveId;
+				UAnimationCurveIdentifierExtensions::GetTransformChildCurveIdentifier(TargetCurveIdentifier, Channel, Axis);
+				SetCurveKey(TargetCurveIdentifier, VectorCurveKeys.ChannelKeys[ChannelIndex], bShouldTransact);
+			}
+		}
+
+		return true;
+	}
+	else
+	{
+		ReportWarningf(LOCTEXT("UnableToFindTransformCurveWarning", "Unable to find transform curve: {0}"), FText::FromName(CurveId.InternalName.DisplayName));
+	}
+
+	return false;
+}
+
+
+bool UAnimDataController::RemoveTransformCurveKey(const FAnimationCurveIdentifier& CurveId, float Time, bool bShouldTransact /*= true*/)
+{
+	ValidateModel();
+
+	FTransformCurve* TransformCurve = Model->FindMutableTransformCurveById(CurveId);
+	if (TransformCurve)
+	{
+		const FString BaseCurveName = CurveId.InternalName.DisplayName.ToString();
+		const TArray<FString> SubCurveNames = { TEXT( "Translation"), TEXT( "Rotation"), TEXT( "Scale") };
+		const TArray<FString> ChannelCurveNames = { TEXT("X"), TEXT("Y"), TEXT("Z") };
+
+		CONDITIONAL_BRACKET(LOCTEXT("RemoveTransformCurveKey_Bracket", "Deleting Animation Transform Curve Key"));
+		
+		for (int32 SubCurveIndex = 0; SubCurveIndex < 3; ++SubCurveIndex)
+		{
+			const ETransformCurveChannel Channel = (ETransformCurveChannel)SubCurveIndex;
+			for (int32 ChannelIndex = 0; ChannelIndex < 3; ++ChannelIndex)
+			{
+				const EVectorCurveChannel Axis = (EVectorCurveChannel)ChannelIndex;
+				FAnimationCurveIdentifier TargetCurveIdentifier = CurveId;
+				UAnimationCurveIdentifierExtensions::GetTransformChildCurveIdentifier(TargetCurveIdentifier, Channel, Axis);
+				RemoveCurveKey(TargetCurveIdentifier, Time, bShouldTransact);
+			}
+		}
+
+
+		return true;
+
+	}
+	else
+	{
+		ReportWarningf(LOCTEXT("UnableToFindTransformCurveWarning", "Unable to find transform curve: {0}"), FText::FromName(CurveId.InternalName.DisplayName));
+	}
+
+	return false;
+}
+
+bool UAnimDataController::RenameCurve(const FAnimationCurveIdentifier& CurveToRenameId, const FAnimationCurveIdentifier& NewCurveId, bool bShouldTransact /*= true*/)
+{
+	ValidateModel();
+
+	if (NewCurveId.IsValid())
+	{
+		if (CurveToRenameId != NewCurveId)
+		{
+			if (CurveToRenameId.CurveType == NewCurveId.CurveType)
+			{
+				FAnimCurveBase* Curve = Model->FindMutableCurveById(CurveToRenameId);
+				if (Curve)
+				{
+					CONDITIONAL_TRANSACTION(LOCTEXT("RenameCurve", "Renaming Curve"));
+
+					FCurveRenamedPayload Payload;
+					Payload.Identifier = FAnimationCurveIdentifier(Curve->Name, CurveToRenameId.CurveType);
+
+					Curve->Name = NewCurveId.InternalName;
+					Payload.NewIdentifier = NewCurveId;
+
+					CONDITIONAL_ACTION(UE::Anim::FRenameCurveAction, NewCurveId, CurveToRenameId);
+
+					Model->Notify(EAnimDataModelNotifyType::CurveRenamed, Payload);
+
+					return true;
+				}
+				else
+				{
+					const FString CurveTypeAsString = GetCurveTypeValueName(CurveToRenameId.CurveType);
+					ReportWarningf(LOCTEXT("UnableToFindCurveWarning", "Unable to find curve: {0} of type {1}"), FText::FromName(CurveToRenameId.InternalName.DisplayName), FText::FromString(CurveTypeAsString));
+				}
+			}
+			else
+			{
+				const FString CurrentCurveTypeAsString = GetCurveTypeValueName(CurveToRenameId.CurveType);
+				const FString NewCurveTypeAsString = GetCurveTypeValueName(NewCurveId.CurveType);
+				ReportWarningf(LOCTEXT("MismatchOfCurveTypesWarning", "Different curve types provided between current and new curve names: {0} ({1}) and {2} ({3})"), FText::FromName(CurveToRenameId.InternalName.DisplayName), FText::FromString(CurrentCurveTypeAsString),
+					FText::FromName(NewCurveId.InternalName.DisplayName), FText::FromString(NewCurveTypeAsString));
+			}
+		}
+		else
+		{
+			ReportWarningf(LOCTEXT("MatchingCurveNamesWarning", "Provided curve names are the same: {0}"), FText::FromName(CurveToRenameId.InternalName.DisplayName));
+		}
+		
+	}
+	else
+	{
+		ReportWarningf(LOCTEXT("InvalidCurveIdentiferProvidedWarning", "Invalid new curve identifier provided: {2} ({3})"), FText::FromName(NewCurveId.InternalName.DisplayName), FText::AsNumber(NewCurveId.InternalName.UID));
+	}
+
+	return false;
+}
+
+bool UAnimDataController::SetCurveColor(const FAnimationCurveIdentifier& CurveId, FLinearColor Color, bool bShouldTransact)
+{
+	ValidateModel();
+
+	if (CurveId.IsValid())
+	{
+		if (CurveId.CurveType == ERawCurveTrackTypes::RCT_Float)
+		{
+			FFloatCurve* Curve = Model->FindMutableFloatCurveById(CurveId);
+			if (Curve)
+			{
+				CONDITIONAL_TRANSACTION(LOCTEXT("ChangingCurveColor", "Changing Curve Color"));
+
+				CONDITIONAL_ACTION(UE::Anim::FSetCurveColorAction, CurveId, Curve->Color);
+
+				Curve->Color = Color;
+
+				FCurveChangedPayload Payload;
+				Payload.Identifier = CurveId;
+				Model->Notify(EAnimDataModelNotifyType::CurveColorChanged, Payload);
+
+				return true;				
+			}
+			else
+			{
+				const FString CurveTypeAsString = GetCurveTypeValueName(CurveId.CurveType);
+				ReportWarningf(LOCTEXT("UnableToFindCurveWarning", "Unable to find curve: {0} of type {1}"), FText::FromName(CurveId.InternalName.DisplayName), FText::FromString(CurveTypeAsString));
+			}
+		}
+		else
+		{
+			ReportWarning(LOCTEXT("NonSupportedCurveColorSetWarning", "Changing curve color is currently only supported for float curves"));
+		}
+	}
+	else
+	{
+		ReportWarningf(LOCTEXT("InvalidCurveIdentifier", "Invalid Curve Identifier : {0} ({1})"), FText::FromName(CurveId.InternalName.DisplayName), FText::AsNumber(CurveId.InternalName.UID));
+	}	
+
+	return false;
+}
+
+bool UAnimDataController::ScaleCurve(const FAnimationCurveIdentifier& CurveId, float Origin, float Factor, bool bShouldTransact /*= true*/)
+{
+	ValidateModel();
+
+	ERawCurveTrackTypes SupportedCurveType = CurveId.CurveType;
+	if (SupportedCurveType == ERawCurveTrackTypes::RCT_Float)
+	{
+		FFloatCurve* Curve = Model->FindMutableFloatCurveById(CurveId);
+		if (Curve)
+		{
+			CONDITIONAL_TRANSACTION(LOCTEXT("ScalingCurve", "Scaling Curve"));
+
+			Curve->FloatCurve.ScaleCurve(Origin, Factor);
+
+			FCurveScaledPayload Payload;
+			Payload.Identifier = CurveId;
+			Payload.Factor = Factor;
+			Payload.Origin = Origin;
+			
+			CONDITIONAL_ACTION(UE::Anim::FScaleCurveAction, CurveId, Origin, 1.0f / Factor, SupportedCurveType);
+
+			Model->Notify(EAnimDataModelNotifyType::CurveScaled, Payload);
+
+			return true;
+		}
+		else
+		{
+			ReportWarningf(LOCTEXT("UnableToFindFloatCurveWarning", "Unable to find float curve: {0}"), FText::FromName(CurveId.InternalName.DisplayName));
+		}
+	}
+	else
+	{
+		ReportWarning(LOCTEXT("NonSupportedCurveScalingWarning", "Scaling curves is currently only supported for float curves"));
+	}
+
+	return false;
+}
+
+bool UAnimDataController::SetCurveKey(const FAnimationCurveIdentifier& CurveId, const FRichCurveKey& Key, bool bShouldTransact)
+{
+	ValidateModel();
+
+	FRichCurve* RichCurve = Model->GetMutableRichCurve(CurveId);
+	ERawCurveTrackTypes SupportedCurveType = CurveId.CurveType;
+	if (RichCurve)
+	{
+		FCurveChangedPayload Payload;
+		Payload.Identifier = CurveId;
+
+		// Set or add rich curve value
+		const FKeyHandle Handle = RichCurve->FindKey(Key.Time, 0.f);
+		if (Handle != FKeyHandle::Invalid())
+		{
+			CONDITIONAL_TRANSACTION(LOCTEXT("SetNamedCurveKey", "Setting Curve Key"));
+			// Cache old value for action
+			const FRichCurveKey CurrentKey = RichCurve->GetKey(Handle);
+			CONDITIONAL_ACTION(UE::Anim::FSetRichCurveKeyAction, CurveId, CurrentKey);
+
+			// Set the new value
+			RichCurve->SetKeyValue(Handle, Key.Value);
+
+			Model->Notify(EAnimDataModelNotifyType::CurveChanged, Payload);
+		}
+		else
+		{
+			CONDITIONAL_TRANSACTION(LOCTEXT("AddNamedCurveKey", "Adding Curve Key"));
+			CONDITIONAL_ACTION(UE::Anim::FRemoveRichCurveKeyAction, CurveId, Key.Time);
+
+			// Add the new key
+			RichCurve->AddKey(Key.Time, Key.Value);
+
+			Model->Notify(EAnimDataModelNotifyType::CurveChanged, Payload);
+		}
+
+		return true;
+	}
+
+	return false;
+}
+
+bool UAnimDataController::RemoveCurveKey(const FAnimationCurveIdentifier& CurveId, float Time, bool bShouldTransact)
+{
+	ValidateModel();
+
+	FRichCurve* RichCurve = Model->GetMutableRichCurve(CurveId);
+	ERawCurveTrackTypes SupportedCurveType = CurveId.CurveType;
+	if (RichCurve)
+	{
+		FCurveChangedPayload Payload;
+		Payload.Identifier = CurveId;
+
+		// Remove key at time value		
+		const FKeyHandle Handle = RichCurve->FindKey(Time, 0.f);
+		if (Handle != FKeyHandle::Invalid())
+		{
+			CONDITIONAL_TRANSACTION(LOCTEXT("RemoveNamedCurveKey", "Removing Curve Key"));
+
+			// Cached current value for action
+			const FRichCurveKey CurrentKey = RichCurve->GetKey(Handle);
+			CONDITIONAL_ACTION(UE::Anim::FAddRichCurveKeyAction, CurveId, CurrentKey);
+
+			RichCurve->DeleteKey(Handle);
+
+			Model->Notify(EAnimDataModelNotifyType::CurveChanged, Payload);
+
+			return true;
+		}
+		else
+		{
+			ReportErrorf(LOCTEXT("RichCurveKeyNotFoundError", "Unable to find rich curve key: curve name {0}, time {1}"), FText::FromName(CurveId.InternalName.DisplayName), FText::AsNumber(Time));
+		}
+	}
+
+	return false;
+}
+
+
+bool UAnimDataController::SetCurveKeys(const FAnimationCurveIdentifier& CurveId, const TArray<FRichCurveKey>& CurveKeys, bool bShouldTransact)
+{
+	ValidateModel();
+
+	FRichCurve* RichCurve = Model->GetMutableRichCurve(CurveId);
+	ERawCurveTrackTypes SupportedCurveType = CurveId.CurveType;
+	if (RichCurve)
+	{
+		CONDITIONAL_TRANSACTION(LOCTEXT("SettingNamedCurveKeys", "Setting Curve Keys"));
+		CONDITIONAL_ACTION(UE::Anim::FSetRichCurveKeysAction, CurveId, RichCurve->GetConstRefOfKeys());
+
+		// Set rich curve values
+		RichCurve->SetKeys(CurveKeys);
+
+		FCurveChangedPayload Payload;
+		Payload.Identifier = CurveId;
+		Model->Notify(EAnimDataModelNotifyType::CurveChanged, Payload);
+
+		return true;
+	}
+
+	return false;
+}
+
+void UAnimDataController::NotifyPopulated()
+{
+	ValidateModel();
+	Model->Notify(EAnimDataModelNotifyType::Populated);
+}
+
+void UAnimDataController::NotifyBracketOpen()
+{
+	ValidateModel();
+	Model->Notify(EAnimDataModelNotifyType::BracketOpened);
+}
+
+void UAnimDataController::NotifyBracketClosed()
+{
+	ValidateModel();
+	Model->Notify(EAnimDataModelNotifyType::BracketClosed);
+}
+
+const bool UAnimDataController::IsSupportedCurveType(ERawCurveTrackTypes CurveType) const
+{
+	const TArray<ERawCurveTrackTypes> SupportedTypes = { ERawCurveTrackTypes::RCT_Float, ERawCurveTrackTypes::RCT_Transform };
+	return SupportedTypes.Contains(CurveType);
+}
+
+void UAnimDataController::ValidateModel() const
+{
+	checkf(Model != nullptr, TEXT("Invalid Model"));
+}
+
+void UAnimDataController::ResizePlayLength(float Length, float T0, float T1, bool bShouldTransact)
+{
+	const TRange<float> PlayRange(TRange<float>::BoundsType::Inclusive(0.f), TRange<float>::BoundsType::Inclusive(Model->PlayLength));
+	if (!FMath::IsNearlyZero(Length) && Length > 0.f)
+	{
+		if (Length != Model->PlayLength)
+		{
+			// Ensure that T0 is within the curent play range
+			if (PlayRange.Contains(T0))
+			{
+				// Ensure that the start and end length of either removal or insertion are valid
+				if (T0 < T1)
+				{
+					CONDITIONAL_TRANSACTION(LOCTEXT("ResizePlayLength", "Resizing Play Length"));
+
+					FSequenceLengthChangedPayload Payload;
+					Payload.T0 = T0;
+					Payload.T1 = T1;
+					Payload.PreviousLength = Model->PlayLength;
+
+					CONDITIONAL_ACTION(UE::Anim::FResizePlayLengthAction, Model, T0, T1);
+
+					Model->PlayLength = Length;
+
+					Model->NumberOfFrames = Model->FrameRate.AsFrameTime(Model->PlayLength).RoundToFrame().Value;
+					Model->NumberOfKeys = Model->NumberOfFrames + 1;
+	
+					Model->Notify<FSequenceLengthChangedPayload>(EAnimDataModelNotifyType::SequenceLengthChanged, Payload);
+				}
+				else
+				{
+					ReportErrorf(LOCTEXT("InvalidEndTimeError", "Invalid T1, smaller that T0 value: T0 {0}, T1 {1}"), FText::AsNumber(T0), FText::AsNumber(T1));
+				}
+			}
+			else
+			{
+				ReportErrorf(LOCTEXT("InvalidStartTimeError", "Invalid T0, not within existing play range: T0 {0}, Play Length {1}"), FText::AsNumber(T0), FText::AsNumber(Model->PlayLength));
+			}
+		}
+		else
+		{
+			ReportWarningf(LOCTEXT("SamePlayLengthWarning", "New play length is same as existing one: {0} seconds"), FText::AsNumber(Length));
+		}
+	}
+	else
+	{
+		ReportErrorf(LOCTEXT("InvalidPlayLengthError", "Invalid play length value provided: {0} seconds"), FText::AsNumber(Length));
+	}
+}
+
+void UAnimDataController::ReportWarning(const FText& InMessage) const
+{
+	FString Message = InMessage.ToString();
+	if (Model != nullptr)
+	{
+		if (UPackage* Package = Cast<UPackage>(Model->GetOutermost()))
+		{
+			Message = FString::Printf(TEXT("%s : %s"), *Package->GetPathName(), *Message);
+		}
+	}
+
+	FScriptExceptionHandler::Get().HandleException(ELogVerbosity::Warning, *Message, *FString());
+}
+
+void UAnimDataController::ReportError(const FText& InMessage) const
+{
+	FString Message = InMessage.ToString();
+	if (Model != nullptr)
+	{
+		if (UPackage* Package = Cast<UPackage>(Model->GetOutermost()))
+		{
+			Message = FString::Printf(TEXT("%s : %s"), *Package->GetPathName(), *Message);
+		}
+	}
+
+	FScriptExceptionHandler::Get().HandleException(ELogVerbosity::Error, *Message, *FString());
+}
+
+FString UAnimDataController::GetCurveTypeValueName(ERawCurveTrackTypes InType) const
+{
+	FString ValueString;
+
+	const UEnum* Enum = FindObject<UEnum>(ANY_PACKAGE, TEXT("ERawCurveTrackTypes"));
+	if (Enum)
+	{
+		ValueString = Enum->GetNameStringByValue((int64)InType);
+	}
+
+	return ValueString;
+}
+
+bool UAnimDataController::CheckOuterClass(UClass* InClass) const
+{
+	ValidateModel();
+	
+	const UObject* ModelOuter = Model->GetOuter();
+	if (ModelOuter)
+	{
+		const UClass* OuterClass = ModelOuter->GetClass();
+		if (OuterClass)
+		{
+			if (OuterClass == InClass || OuterClass->IsChildOf(InClass))
+			{
+				return true;
+			}
+			else
+			{
+				ReportErrorf(LOCTEXT("NoValidOuterClassError", "Incorrect outer object class found for Animation Data Model {0}, expected {1} actual {2}"), FText::FromString(Model->GetName()), FText::FromString(InClass->GetName()), FText::FromString(OuterClass->GetName()));
+			}
+		}
+	}
+	else
+	{
+		ReportErrorf(LOCTEXT("NoValidOuterObjectFoundError", "No valid outer object found for Animation Data Model {0}"), FText::FromString(Model->GetName()));
+	}
+
+	return false;
+}
+
+int32 UAnimDataController::AddBoneTrack(FName BoneName, bool bShouldTransact /*= true*/)
+{
+	if (!CheckOuterClass(UAnimSequence::StaticClass()))
+	{
+		return INDEX_NONE;
+	}
+
+	CONDITIONAL_TRANSACTION(LOCTEXT("AddBoneTrack", "Adding Animation Data Track"));
+	return InsertBoneTrack(BoneName, INDEX_NONE, bShouldTransact);
+}
+
+int32 UAnimDataController::InsertBoneTrack(FName BoneName, int32 DesiredIndex, bool bShouldTransact /*= true*/)
+{
+	if (!CheckOuterClass(UAnimSequence::StaticClass()))
+	{
+		return INDEX_NONE;
+	}
+	
+	const int32 TrackIndex = Model->GetBoneTrackIndexByName(BoneName);
+
+	if (TrackIndex == INDEX_NONE)
+	{
+		if (Model->GetNumBoneTracks() >= MAX_ANIMATION_TRACKS)
+		{
+			ReportWarningf(LOCTEXT("MaxNumberOfTracksReachedWarning", "Cannot add track with name {0}. An animation sequence cannot contain more than 65535 tracks"), FText::FromName(BoneName));
+		}
+		else
+		{
+			CONDITIONAL_TRANSACTION(LOCTEXT("InsertBoneTrack", "Inserting Animation Data Track"));
+
+			// Determine correct index to do insertion at
+			const int32 InsertIndex = Model->BoneAnimationTracks.IsValidIndex(DesiredIndex) ? DesiredIndex : Model->BoneAnimationTracks.Num();
+
+			FBoneAnimationTrack& NewTrack = Model->BoneAnimationTracks.InsertDefaulted_GetRef(InsertIndex);
+			NewTrack.Name = BoneName;
+
+			if (const UAnimSequence* AnimationSequence = Model->GetAnimationSequence())
+			{
+				if (const USkeleton* Skeleton = AnimationSequence->GetSkeleton())
+				{
+					const int32 BoneIndex = Skeleton->GetReferenceSkeleton().FindBoneIndex(BoneName);
+
+					if (BoneIndex == INDEX_NONE)
+					{
+						ReportWarningf(LOCTEXT("UnableToFindBoneIndexWarning", "Unable to retrieve bone index for track: {0}"), FText::FromName(BoneName));
+					}
+
+					NewTrack.BoneTreeIndex = BoneIndex;
+				}
+				else
+				{
+					ReportError(LOCTEXT("UnableToGetOuterSkeletonError", "Unable to retrieve Skeleton for outer Animation Sequence"));
+				}
+			}
+			else
+			{
+				ReportError(LOCTEXT("UnableToGetOuterAnimSequenceError", "Unable to retrieve outer Animation Sequence"));
+			}
+
+			FAnimationTrackAddedPayload Payload;
+			Payload.Name = BoneName;
+			Payload.TrackIndex = InsertIndex;
+
+			Model->Notify<FAnimationTrackAddedPayload>(EAnimDataModelNotifyType::TrackAdded, Payload);
+			CONDITIONAL_ACTION(UE::Anim::FRemoveTrackAction, NewTrack, InsertIndex);
+
+			return InsertIndex;
+		}
+	}
+	else
+	{
+		ReportWarningf(LOCTEXT("TrackNameAlreadyExistsWarning", "Track with name {0} already exists"), FText::FromName(BoneName));
+	}
+	
+	return TrackIndex;
+}
+
+bool UAnimDataController::RemoveBoneTrack(FName BoneName, bool bShouldTransact /*= true*/)
+{
+	if (!CheckOuterClass(UAnimSequence::StaticClass()))
+	{
+		return false;
+	}
+
+	const FBoneAnimationTrack* ExistingTrackPtr = Model->FindBoneTrackByName(BoneName);
+
+	if (ExistingTrackPtr != nullptr)
+	{
+		CONDITIONAL_TRANSACTION(LOCTEXT("RemoveBoneTrack", "Removing Animation Data Track"));
+		const int32 TrackIndex = Model->BoneAnimationTracks.IndexOfByPredicate([ExistingTrackPtr](const FBoneAnimationTrack& Track)
+		{
+			return Track.Name == ExistingTrackPtr->Name;
+		});
+
+		ensure(TrackIndex != INDEX_NONE);
+
+		CONDITIONAL_ACTION(UE::Anim::FAddTrackAction, *ExistingTrackPtr, TrackIndex);
+		Model->BoneAnimationTracks.RemoveAt(TrackIndex);
+
+		FAnimationTrackRemovedPayload Payload;
+		Payload.Name = BoneName;
+
+		Model->Notify(EAnimDataModelNotifyType::TrackRemoved, Payload);
+
+		return true;
+	}
+	else
+	{
+		ReportWarningf(LOCTEXT("UnableToFindTrackWarning", "Could not find track with name {0}"), FText::FromName(BoneName));
+	}
+
+	return false;
+}
+
+void UAnimDataController::RemoveAllBoneTracks(bool bShouldTransact /*= true*/)
+{
+	if (!CheckOuterClass(UAnimSequence::StaticClass()))
+	{
+		return;
+	}
+	
+	TArray<FName> TrackNames;
+	Model->GetBoneTrackNames(TrackNames);
+
+	if (TrackNames.Num())
+	{
+		CONDITIONAL_BRACKET(LOCTEXT("RemoveAllBoneTracks", "Removing all Animation Data Tracks"));
+		for (const FName& TrackName : TrackNames)
+		{
+			RemoveBoneTrack(TrackName, bShouldTransact);
+		}
+	}	
+}
+
+bool UAnimDataController::SetBoneTrackKeys(FName BoneName, const TArray<FVector>& PositionalKeys, const TArray<FQuat>& RotationalKeys, const TArray<FVector>& ScalingKeys, bool bShouldTransact /*= true*/)
+{
+	if (!CheckOuterClass(UAnimSequence::StaticClass()))
+	{
+		return false;
+	}
+
+	CONDITIONAL_TRANSACTION(LOCTEXT("SetTrackKeysTransaction", "Setting Animation Data Track keys"));
+
+	// Validate key format
+	const int32 MaxNumKeys = FMath::Max(FMath::Max(PositionalKeys.Num(), RotationalKeys.Num()), ScalingKeys.Num());
+
+	if (MaxNumKeys > 0)
+	{
+		const bool bValidPosKeys = PositionalKeys.Num() == MaxNumKeys;
+		const bool bValidRotKeys = RotationalKeys.Num() == MaxNumKeys;
+		const bool bValidScaleKeys = ScalingKeys.Num() == MaxNumKeys;
+
+		if (bValidPosKeys && bValidRotKeys && bValidScaleKeys)
+		{
+			if (FBoneAnimationTrack* TrackPtr = Model->FindMutableBoneTrackByName(BoneName))
+			{
+				CONDITIONAL_ACTION(UE::Anim::FSetTrackKeysAction, *TrackPtr);
+
+				TrackPtr->InternalTrackData.PosKeys = PositionalKeys;
+				TrackPtr->InternalTrackData.RotKeys = RotationalKeys;
+				TrackPtr->InternalTrackData.ScaleKeys = ScalingKeys;
+
+
+				FAnimationTrackChangedPayload Payload;
+				Payload.Name = BoneName;
+
+				Model->Notify(EAnimDataModelNotifyType::TrackChanged, Payload);
+
+				return true;
+			}
+			else
+			{
+				ReportWarningf(LOCTEXT("InvalidTrackNameWarning", "Track with name {0} does not exist"), FText::FromName(BoneName));
+			}
+		}
+		else
+		{
+			ReportErrorf(LOCTEXT("InvalidTrackKeyDataError", "Invalid track key data, expected uniform data: number of positional keys {0}, number of rotational keys {1}, number of scaling keys {2}"), FText::AsNumber(PositionalKeys.Num()), FText::AsNumber(RotationalKeys.Num()), FText::AsNumber(ScalingKeys.Num()));
+		}
+	}
+	else
+	{
+		ReportErrorf(LOCTEXT("MissingTrackKeyDataError", "Missing track key data, expected uniform data: number of positional keys {0}, number of rotational keys {1}, number of scaling keys {2}"), FText::AsNumber(PositionalKeys.Num()), FText::AsNumber(RotationalKeys.Num()), FText::AsNumber(ScalingKeys.Num()));
+	}
+
+	return false;
+}
+
+void UAnimDataController::ResizeCurves(float NewLength, bool bInserted, float T0, float T1, bool bShouldTransact /*= true*/)
+{
+	CONDITIONAL_BRACKET(LOCTEXT("ResizeCurves", "Resizing all Curves"));
+
+	for (FFloatCurve& Curve : Model->CurveData.FloatCurves)
+	{
+		FFloatCurve ResizedCurve = Curve;
+		ResizedCurve.Resize(NewLength, bInserted, T0, T1);
+		SetCurveKeys(FAnimationCurveIdentifier(Curve.Name, ERawCurveTrackTypes::RCT_Float), ResizedCurve.FloatCurve.GetConstRefOfKeys(), bShouldTransact);
+	}
+
+	for (FTransformCurve& Curve : Model->CurveData.TransformCurves)
+	{
+		FTransformCurve ResizedCurve = Curve;
+		for (int32 SubCurveIndex = 0; SubCurveIndex < 3; ++SubCurveIndex)
+		{
+			const ETransformCurveChannel Channel = (ETransformCurveChannel)SubCurveIndex;
+			FVectorCurve& SubCurve = *ResizedCurve.GetVectorCurveByIndex(SubCurveIndex);
+			for (int32 ChannelIndex = 0; ChannelIndex < 3; ++ChannelIndex)
+			{
+				const EVectorCurveChannel Axis = (EVectorCurveChannel)ChannelIndex;
+				FAnimationCurveIdentifier TargetCurveIdentifier = FAnimationCurveIdentifier(Curve.Name, ERawCurveTrackTypes::RCT_Transform);
+				UAnimationCurveIdentifierExtensions::GetTransformChildCurveIdentifier(TargetCurveIdentifier, Channel, Axis);
+				
+				FRichCurve& ChannelCurve = SubCurve.FloatCurves[ChannelIndex];
+				ChannelCurve.ReadjustTimeRange(0, NewLength, bInserted, T0, T1);
+				SetCurveKeys(TargetCurveIdentifier, ChannelCurve.GetConstRefOfKeys(), bShouldTransact);
+			}
+		}
+	}
+}
+
+void UAnimDataController::ResizeAttributes(float NewLength, bool bInserted, float T0, float T1, bool bShouldTransact)
+{
+	CONDITIONAL_BRACKET(LOCTEXT("ResizeAttributes", "Resizing all Attributes"));
+
+	for (FAnimatedBoneAttribute& Attribute : Model->AnimatedBoneAttributes)
+	{
+		FAttributeCurve ResizedCurve = Attribute.Curve;
+		ResizedCurve.ReadjustTimeRange(0, NewLength, bInserted, T0, T1);
+
+		// Generate arrays necessary for API
+		TArray<float> Times;
+		TArray<const void*> Values;
+		const TArray<FAttributeKey>& Keys = ResizedCurve.GetConstRefOfKeys();
+		for (int32 KeyIndex = 0; KeyIndex < ResizedCurve.GetNumKeys(); ++KeyIndex)
+		{
+			Times.Add(Keys[KeyIndex].Time);
+			Values.Add(Keys[KeyIndex].GetValuePtr<void>());
+		}
+		
+		SetAttributeKeys(Attribute.Identifier, Times, Values, bShouldTransact);
+	}
+}
+
+bool UAnimDataController::AddAttribute(const FAnimationAttributeIdentifier& AttributeIdentifier, bool bShouldTransact /*= true*/)
+{
+	if (AttributeIdentifier.IsValid())
+	{
+		const bool bAttributeAlreadyExists = Model->AnimatedBoneAttributes.ContainsByPredicate([AttributeIdentifier](const FAnimatedBoneAttribute& Attribute) -> bool
+		{
+			return Attribute.Identifier == AttributeIdentifier;
+		});
+
+		if (!bAttributeAlreadyExists)
+		{
+			CONDITIONAL_TRANSACTION(LOCTEXT("AddAttribute", "Adding Animated Bone Attribute"));
+
+			FAnimatedBoneAttribute& Attribute = Model->AnimatedBoneAttributes.AddDefaulted_GetRef();
+			Attribute.Identifier = AttributeIdentifier;
+
+			Attribute.Curve.SetScriptStruct(AttributeIdentifier.GetType());
+		
+			CONDITIONAL_ACTION(UE::Anim::FRemoveAtributeAction, AttributeIdentifier);
+
+			FAttributeAddedPayload Payload;
+			Payload.Identifier = AttributeIdentifier;
+			Model->Notify(EAnimDataModelNotifyType::AttributeAdded, Payload);
+
+			return true;
+		}
+		else
+		{
+			ReportErrorf(LOCTEXT("AttributeAlreadyExists", "Attribute identifier provided already exists: {0} {1} ({2}) {3}"),
+			FText::FromName(AttributeIdentifier.GetName()), FText::FromName(AttributeIdentifier.GetBoneName()), FText::AsNumber(AttributeIdentifier.GetBoneIndex()), FText::FromName(AttributeIdentifier.GetType()->GetFName()));
+		}
+	}
+	else
+	{
+		ReportError(LOCTEXT("InvalidAttributeIdentifier", "Invalid attribute identifier provided"));
+	}
+
+	return false;
+}
+
+bool UAnimDataController::RemoveAttribute(const FAnimationAttributeIdentifier& AttributeIdentifier, bool bShouldTransact /*= true*/)
+{
+	if (AttributeIdentifier.IsValid())
+	{
+		const int32 AttributeIndex = Model->AnimatedBoneAttributes.IndexOfByPredicate([AttributeIdentifier](const FAnimatedBoneAttribute& Attribute) -> bool
+		{
+			return Attribute.Identifier == AttributeIdentifier;
+		});
+
+		if (AttributeIndex != INDEX_NONE)
+		{
+			CONDITIONAL_TRANSACTION(LOCTEXT("RemoveAttribute", "Removing Animated Bone Attribute"));
+
+			CONDITIONAL_ACTION(UE::Anim::FAddAtributeAction, Model->AnimatedBoneAttributes[AttributeIndex]);
+
+			Model->AnimatedBoneAttributes.RemoveAtSwap(AttributeIndex);
+			
+			FAttributeRemovedPayload Payload;
+			Payload.Identifier = AttributeIdentifier;
+			Model->Notify(EAnimDataModelNotifyType::AttributeRemoved, Payload);
+
+			return true;
+		}
+		else
+		{
+			ReportErrorf(LOCTEXT("AttributeNotFound", "Attribute identifier provided was not found: {0} {1} ({2}) {3}"),
+				FText::FromName(AttributeIdentifier.GetName()), FText::FromName(AttributeIdentifier.GetBoneName()), FText::AsNumber(AttributeIdentifier.GetBoneIndex()), FText::FromName(AttributeIdentifier.GetType()->GetFName()));
+		}
+	}
+	else
+	{
+		ReportError(LOCTEXT("InvalidAttributeIdentifier", "Invalid attribute identifier provided"));
+	}
+
+	return false;
+}
+
+int32 UAnimDataController::RemoveAllAttributesForBone(const FName& BoneName, bool bShouldTransact)
+{
+	int32 NumRemovedAttributes = 0;
+
+	// Generate list of attribute identifiers, matching the bone name, for removal
+	TArray<FAnimationAttributeIdentifier> Identifiers;
+	Algo::TransformIf(Model->AnimatedBoneAttributes, Identifiers,
+		[BoneName](const FAnimatedBoneAttribute& Attribute) -> bool
+		{
+			return Attribute.Identifier.GetBoneName() == BoneName;
+		},
+		[](const FAnimatedBoneAttribute& Attribute)
+		{
+			return Attribute.Identifier;
+		}
+	);
+
+	if (Identifiers.Num())
+	{
+		CONDITIONAL_BRACKET(LOCTEXT("RemoveAllAttributesForBone", "Removing all Attributes for Bone"));
+		for (const FAnimationAttributeIdentifier& Identifier : Identifiers)
+		{
+			NumRemovedAttributes += RemoveAttribute(Identifier, bShouldTransact) ? 1 : 0;
+		}
+	}	
+
+	return NumRemovedAttributes;
+}
+
+
+int32 UAnimDataController::RemoveAllAttributes(bool bShouldTransact)
+{
+	int32 NumRemovedAttributes = 0;
+
+	TArray<FAnimationAttributeIdentifier> Identifiers;
+	Algo::Transform(Model->AnimatedBoneAttributes, Identifiers, [](const FAnimatedBoneAttribute& Attribute)
+	{
+		return Attribute.Identifier;
+	});
+
+	if (Identifiers.Num())
+	{
+		CONDITIONAL_BRACKET(LOCTEXT("RemoveAllAttributes", "Removing all Attributes"));
+		for (const FAnimationAttributeIdentifier& Identifier : Identifiers)
+		{
+			NumRemovedAttributes += RemoveAttribute(Identifier, bShouldTransact) ? 1: 0;
+		}
+	}
+
+	return NumRemovedAttributes;
+}
+
+bool UAnimDataController::SetAttributeKey_Internal(const FAnimationAttributeIdentifier& AttributeIdentifier, float Time, const void* KeyValue, const UScriptStruct* TypeStruct, bool bShouldTransact /*= true*/)
+{
+	if (AttributeIdentifier.IsValid())
+	{
+		if (KeyValue)
+		{
+			FAnimatedBoneAttribute* AttributePtr = Model->AnimatedBoneAttributes.FindByPredicate([AttributeIdentifier](FAnimatedBoneAttribute& Attribute)
+			{
+				return Attribute.Identifier == AttributeIdentifier;
+			});
+
+			if (AttributePtr)
+			{
+				if (TypeStruct == AttributePtr->Identifier.GetType())
+				{
+					CONDITIONAL_TRANSACTION(LOCTEXT("SettingAttributeKey", "Setting Animated Bone Attribute key"));
+
+					FAttributeCurve& Curve = AttributePtr->Curve;
+					FKeyHandle KeyHandle = Curve.FindKey(Time);
+					// In case the key does not yet exist one will be added, and thus the undo is a remove
+					if (KeyHandle == FKeyHandle::Invalid())
+					{
+						CONDITIONAL_ACTION(UE::Anim::FRemoveAtributeKeyAction, AttributeIdentifier, Time);
+						Curve.UpdateOrAddKey(Time, KeyValue);
+					}
+					// In case the key does exist it will be updated , and thus the undo is a revert to the current value
+					else
+					{
+						CONDITIONAL_ACTION(UE::Anim::FSetAtributeKeyAction, AttributeIdentifier, Curve.GetKey(KeyHandle));
+						Curve.UpdateOrAddKey(Time, KeyValue);
+					}
+
+					FAttributeChangedPayload Payload;
+					Payload.Identifier = AttributeIdentifier;
+					Model->Notify(EAnimDataModelNotifyType::AttributeChanged, Payload);
+
+					return true;
+				}
+				else
+				{
+					ReportErrorf(LOCTEXT("AttributeTypeDoesNotMatchKeyType", "Key type does not match attribute: {0} {1}"),
+						FText::FromName(AttributePtr->Identifier.GetType()->GetFName()), FText::FromName(TypeStruct->GetFName()));
+				}
+			}
+			else
+			{
+				ReportErrorf(LOCTEXT("AttributeNotFound", "Attribute identifier provided was not found: {0} {1} ({2}) {3}"),
+					FText::FromName(AttributeIdentifier.GetName()), FText::FromName(AttributeIdentifier.GetBoneName()), FText::AsNumber(AttributeIdentifier.GetBoneIndex()), FText::FromName(AttributeIdentifier.GetType()->GetFName()));
+			}
+		}
+		else
+		{
+			ReportError(LOCTEXT("InvalidAttributeKey", "Invalid attribute key value provided"));
+		}
+	}
+	else
+	{
+		ReportError(LOCTEXT("InvalidAttributeIdentifier", "Invalid attribute identifier provided"));
+	}
+
+	return false;
+}
+
+bool UAnimDataController::SetAttributeKeys_Internal(const FAnimationAttributeIdentifier& AttributeIdentifier, TArrayView<const float> Times, TArrayView<const void*> KeyValues, const UScriptStruct* TypeStruct, bool bShouldTransact)
+{
+	if (AttributeIdentifier.IsValid())
+	{
+		if (Times.Num() == KeyValues.Num())
+		{
+			FAnimatedBoneAttribute* AttributePtr = Model->AnimatedBoneAttributes.FindByPredicate([AttributeIdentifier](FAnimatedBoneAttribute& Attribute)
+			{
+				return Attribute.Identifier == AttributeIdentifier;
+			});
+
+			if (AttributePtr)
+			{
+				if (TypeStruct == AttributePtr->Identifier.GetType())
+				{
+					CONDITIONAL_TRANSACTION(LOCTEXT("SettingAttributeKeys", "Setting Animated Bone Attribute keys"));
+
+					FAnimatedBoneAttribute& Attribute = *AttributePtr;
+
+					CONDITIONAL_ACTION(UE::Anim::FSetAtributeKeysAction, Attribute);
+			
+					Attribute.Curve.SetKeys(Times, KeyValues);
+
+					FAttributeChangedPayload Payload;
+					Payload.Identifier = AttributeIdentifier;
+					Model->Notify(EAnimDataModelNotifyType::AttributeChanged, Payload);
+
+					return true;
+				}
+				else
+				{
+					ReportErrorf(LOCTEXT("AttributeTypeDoesNotMatchKeyType", "Key type does not match attribute: {0} {1}"),
+						FText::FromName(AttributePtr->Identifier.GetType()->GetFName()), FText::FromName(TypeStruct->GetFName()));
+				}
+			}
+			else
+			{
+				ReportErrorf(LOCTEXT("AttributeNotFound", "Attribute identifier provided was not found: {0} {1} ({2}) {3}"),
+					FText::FromName(AttributeIdentifier.GetName()), FText::FromName(AttributeIdentifier.GetBoneName()), FText::AsNumber(AttributeIdentifier.GetBoneIndex()), FText::FromName(AttributeIdentifier.GetType()->GetFName()));
+			}
+		}
+		else
+		{
+			ReportErrorf(LOCTEXT("AttributeKeysMismatch", "Non matching number of key time/values: time entries {0} key entries {1}"),
+				FText::AsNumber(Times.Num()), FText::AsNumber(KeyValues.Num()));
+		}
+	}
+	else
+	{
+		ReportError(LOCTEXT("InvalidAttributeIdentifier", "Invalid attribute identifier provided"));
+	}
+
+	return false;
+}
+
+
+bool UAnimDataController::RemoveAttributeKey(const FAnimationAttributeIdentifier& AttributeIdentifier, float Time, bool bShouldTransact /*= true*/)
+{
+	if (AttributeIdentifier.IsValid())
+	{
+		FAnimatedBoneAttribute* AttributePtr = Model->AnimatedBoneAttributes.FindByPredicate([AttributeIdentifier](const FAnimatedBoneAttribute& Attribute)
+		{
+			return Attribute.Identifier == AttributeIdentifier;
+		});
+
+		if (AttributePtr)
+		{
+			FAttributeCurve& Curve = AttributePtr->Curve;
+			FKeyHandle KeyHandle = Curve.FindKey(Time);
+
+			if (KeyHandle != FKeyHandle::Invalid())
+			{
+				CONDITIONAL_TRANSACTION(LOCTEXT("RemovingAttributeKey", "Removing Animated Bone Attribute key"));
+
+				CONDITIONAL_ACTION(UE::Anim::FAddAtributeKeyAction, AttributeIdentifier, Curve.GetKey(KeyHandle));
+
+				Curve.DeleteKey(KeyHandle);
+
+				FAttributeAddedPayload Payload;
+				Payload.Identifier = AttributeIdentifier;
+				Model->Notify(EAnimDataModelNotifyType::AttributeChanged, Payload);
+
+				return true;
+			}
+			else
+			{
+				ReportWarning(LOCTEXT("AttributeKeyNotFound", "Attribute does not contain key for provided time"));
+			}	
+		}
+		else
+		{
+			ReportErrorf(LOCTEXT("AttributeNotFound", "Attribute identifier provided was not found: {0} {1} ({2}) {3}"),
+				FText::FromName(AttributeIdentifier.GetName()), FText::FromName(AttributeIdentifier.GetBoneName()), FText::AsNumber(AttributeIdentifier.GetBoneIndex()), FText::FromName(AttributeIdentifier.GetType()->GetFName()));
+		}
+	}
+	else
+	{
+		ReportError(LOCTEXT("InvalidAttributeIdentifier", "Invalid attribute identifier provided"));
+	}
+
+	return false;
+}
+
+bool UAnimDataController::DuplicateAttribute(const FAnimationAttributeIdentifier& AttributeIdentifier, const FAnimationAttributeIdentifier& NewAttributeIdentifier, bool bShouldTransact)
+{
+	ValidateModel();
+
+	if (AttributeIdentifier.IsValid() && NewAttributeIdentifier.IsValid())
+	{
+		if (AttributeIdentifier.GetType() == NewAttributeIdentifier.GetType())
+		{
+			const FAnimatedBoneAttribute* ExistingAttributePtr = Model->FindAttribute(NewAttributeIdentifier);
+			if (ExistingAttributePtr == nullptr)
+			{
+				if(const FAnimatedBoneAttribute* AttributePtr = Model->FindAttribute(AttributeIdentifier))
+				{
+					CONDITIONAL_TRANSACTION(LOCTEXT("DuplicateAttribute", "Duplicating Animation Attribute"));
+
+					FAnimatedBoneAttribute& DuplicateAttribute = Model->AnimatedBoneAttributes.AddDefaulted_GetRef();
+					DuplicateAttribute.Identifier = NewAttributeIdentifier;
+					DuplicateAttribute.Curve = AttributePtr->Curve;
+
+					FAttributeAddedPayload Payload;
+					Payload.Identifier = NewAttributeIdentifier;
+					Model->Notify(EAnimDataModelNotifyType::AttributeAdded, Payload);
+
+					CONDITIONAL_ACTION(UE::Anim::FRemoveAtributeAction, NewAttributeIdentifier);
+					
+					return true;
+				}
+				else
+				{
+					ReportWarningf(LOCTEXT("AttributeNameToDuplicateNotFoundWarning", "Could not find attribute with name {0} and type {1} for duplication"), FText::FromName(AttributeIdentifier.GetName()), FText::FromString(AttributeIdentifier.GetType()->GetName()));
+				}
+			}
+			else
+			{
+				ReportWarningf(LOCTEXT("ExistingAttributeWarning", "Attribute with name {0} already exists"), FText::FromName(AttributeIdentifier.GetName()));
+			}
+		}
+		else
+		{
+			ReportWarningf(LOCTEXT("InvalidAttributeTypeWarning", "Attribute types do not match: {0} ({1})"), FText::FromString(AttributeIdentifier.GetType()->GetName()), FText::FromString(AttributeIdentifier.GetType()->GetName()));
+		}
+	}
+	else
+	{
+		ReportWarningf(LOCTEXT("InvalidAttributeIdentifierWarning", "Invalid attribute identifier(s) provided: {0} and/or {1}"), FText::FromName(AttributeIdentifier.GetName()), FText::FromName(AttributeIdentifier.GetName()));
+	}
+
+	return false;
+}
+
+#endif // WITH_EDITOR
+
+#undef LOCTEXT_NAMESPACE // "AnimDataController"
+
diff --git a/Engine/Source/Developer/AnimationDataController/Private/AnimDataControllerActions.cpp b/Engine/Source/Developer/AnimationDataController/Private/AnimDataControllerActions.cpp
new file mode 100644
--- /dev/null
+++ b/Engine/Source/Developer/AnimationDataController/Private/AnimDataControllerActions.cpp
@@ -0,0 +1,539 @@
+// Copyright Epic Games, Inc. All Rights Reserved.
+
+#include "AnimDataControllerActions.h"
+#include "AnimDataController.h"
+
+#include "Animation/AnimData/AnimDataModel.h"
+#include "UObject/StrongObjectPtr.h"
+
+#define LOCTEXT_NAMESPACE "AnimDataControllerActions"
+
+#if WITH_EDITOR
+
+namespace UE {
+
+namespace Anim {
+
+TUniquePtr<FChange> FAnimDataBaseAction::Execute(UObject* Object)
+{
+	UAnimDataModel* Model = Cast<UAnimDataModel>(Object);
+	checkf(Model != nullptr, TEXT("Invalid UAnimDataModel Object"));
+
+	TStrongObjectPtr<UAnimDataController> Controller(NewObject<UAnimDataController>());
+	Controller->SetModel(Model);
+
+	return ExecuteInternal(Model, Controller.Get());
+}
+
+FString FAnimDataBaseAction::ToString() const
+{
+	return ToStringInternal();
+}
+
+TUniquePtr<FChange> FOpenBracketAction::ExecuteInternal(UAnimDataModel* Model, UAnimDataController* Controller)
+{
+	Controller->NotifyBracketOpen();
+	return MakeUnique<FCloseBracketAction>(Description);
+}
+
+FString FOpenBracketAction::ToStringInternal() const
+{
+	return FString::Printf(TEXT("Open Bracket: %s"), *Description);
+}
+
+TUniquePtr<FChange> FCloseBracketAction::ExecuteInternal(UAnimDataModel* Model, UAnimDataController* Controller)
+{
+	Controller->NotifyBracketClosed();
+	return MakeUnique<FOpenBracketAction>(Description);
+}
+
+FString FCloseBracketAction::ToStringInternal() const
+{
+	return TEXT("Closing Bracket");
+}
+
+FAddTrackAction::FAddTrackAction(const FBoneAnimationTrack& Track, int32 InTrackIndex)
+{
+	Name = Track.Name;
+	BoneTreeIndex = Track.BoneTreeIndex;
+	Data = Track.InternalTrackData;
+	TrackIndex = InTrackIndex;
+}
+
+TUniquePtr<FChange> FAddTrackAction::ExecuteInternal(UAnimDataModel* Model, UAnimDataController* Controller)
+{
+	Controller->InsertBoneTrack(Name, TrackIndex, false);
+	
+	if (Data.PosKeys.Num() || Data.RotKeys.Num() || Data.ScaleKeys.Num())
+	{
+		Controller->SetBoneTrackKeys(Name, Data.PosKeys, Data.RotKeys, Data.ScaleKeys, false);
+	}
+	
+	const FBoneAnimationTrack& AddedTrack = Model->GetBoneTrackByName(Name);
+	return MakeUnique<FRemoveTrackAction>(AddedTrack, TrackIndex);
+}
+
+FString FAddTrackAction::ToStringInternal() const
+{
+	return FText::Format(LOCTEXT("AddTrackAction_Description", "Adding animation bone track '{0}'."), FText::FromName(Name)).ToString();
+}
+
+FRemoveTrackAction::FRemoveTrackAction(const FBoneAnimationTrack& Track, int32 InTrackIndex)
+{
+	Name = Track.Name;
+	TrackIndex = InTrackIndex;
+}
+
+TUniquePtr<FChange> FRemoveTrackAction::ExecuteInternal(UAnimDataModel* Model, UAnimDataController* Controller)
+{
+	const FBoneAnimationTrack& Track = Model->GetBoneTrackByName(Name);
+
+	TUniquePtr<FChange> InverseAction = MakeUnique<FAddTrackAction>(Track, TrackIndex);
+	Controller->RemoveBoneTrack(Name, false);
+	return InverseAction;
+}
+
+FString FRemoveTrackAction::ToStringInternal() const
+{
+	return FText::Format(LOCTEXT("RemoveTrackAction_Description", "Removing animation bone Track '{0}'."), FText::FromName(Name)).ToString();
+}
+
+FSetTrackKeysAction::FSetTrackKeysAction(const FBoneAnimationTrack& Track)
+{
+	Name = Track.Name;
+
+	TrackData = Track.InternalTrackData;
+}
+
+TUniquePtr<FChange> FSetTrackKeysAction::ExecuteInternal(UAnimDataModel* Model, UAnimDataController* Controller)
+{
+	const FBoneAnimationTrack& Track = Model->GetBoneTrackByName(Name);
+	TUniquePtr<FChange> InverseAction = MakeUnique<FSetTrackKeysAction>(Track);
+
+	Controller->SetBoneTrackKeys(Name, TrackData.PosKeys, TrackData.RotKeys, TrackData.ScaleKeys, false);
+	return InverseAction;
+}
+
+FString FSetTrackKeysAction::ToStringInternal() const
+{
+	return FText::Format(LOCTEXT("SetTrackKeysAction_Description", "Setting keys for animation bone track '{0}'."), FText::FromName(Name)).ToString();
+}
+
+FResizePlayLengthAction::FResizePlayLengthAction(const UAnimDataModel* InModel, float t0, float t1) : T0(t0), T1(t1)
+{
+	Length = InModel->GetPlayLength();
+}
+
+TUniquePtr<FChange> FResizePlayLengthAction::ExecuteInternal(UAnimDataModel* Model, UAnimDataController* Controller)
+{
+	TUniquePtr<FChange> InverseAction = MakeUnique<FResizePlayLengthAction>(Model, T0, T1);
+	Controller->ResizePlayLength(Length, T0, T1, false);
+	return InverseAction;
+}
+
+FString FResizePlayLengthAction::ToStringInternal() const
+{
+	return FText::Format(LOCTEXT("ResizePlayLengthAction_Description", "Resizing play length to {0}."), FText::AsNumber(Length)).ToString();
+}
+
+FSetFrameRateAction::FSetFrameRateAction(const UAnimDataModel* InModel)
+{
+	FrameRate = InModel->GetFrameRate();
+}
+
+TUniquePtr<FChange> FSetFrameRateAction::ExecuteInternal(UAnimDataModel* Model, UAnimDataController* Controller)
+{
+	TUniquePtr<FChange> InverseAction = MakeUnique<FSetFrameRateAction>(Model);
+	Controller->SetFrameRate(FrameRate, false);
+	return InverseAction;
+}
+
+FString FSetFrameRateAction::ToStringInternal() const
+{
+	return FText::Format(LOCTEXT("SetFrameRateAction_Description", "Setting Frame Rate to {0}."), FrameRate.ToPrettyText()).ToString();
+}
+
+TUniquePtr<FChange> FAddCurveAction::ExecuteInternal(UAnimDataModel* Model, UAnimDataController* Controller)
+{
+	TUniquePtr<FChange> InverseAction = MakeUnique<FRemoveCurveAction>(CurveId);
+
+	Controller->AddCurve(CurveId, Flags, false);
+
+	return InverseAction;
+}
+
+FString FAddCurveAction::ToStringInternal() const
+{
+	const FString FloatLabel(TEXT("float"));
+	const FString TransformLabel(TEXT("transform"));
+
+	return FText::Format(LOCTEXT("AddCurveAction_Description", "Adding {0} curve '{1}'."), FText::FromString(CurveId.CurveType == ERawCurveTrackTypes::RCT_Float ? FloatLabel : TransformLabel), FText::FromName(CurveId.InternalName.DisplayName)).ToString();
+}
+
+TUniquePtr<FChange> FRemoveCurveAction::ExecuteInternal(UAnimDataModel* Model, UAnimDataController* Controller)
+{
+	TUniquePtr<FChange> InverseAction;	
+
+	if (CurveId.CurveType == ERawCurveTrackTypes::RCT_Float)
+	{
+		const FFloatCurve& Curve = Model->GetFloatCurve(CurveId);
+		InverseAction = MakeUnique<FAddFloatCurveAction>(CurveId, Curve.GetCurveTypeFlags(), Curve.FloatCurve.GetConstRefOfKeys(), Curve.Color);
+	}
+	else if (CurveId.CurveType == ERawCurveTrackTypes::RCT_Transform)
+	{
+		const FTransformCurve& Curve = Model->GetTransformCurve(CurveId);
+		InverseAction = MakeUnique<FAddTransformCurveAction>(CurveId, Curve.GetCurveTypeFlags(), Curve);
+	}
+
+	Controller->RemoveCurve(CurveId, false);
+
+	return InverseAction;
+}
+
+FString FRemoveCurveAction::ToStringInternal() const
+{
+	const FString FloatLabel(TEXT("float"));
+	const FString TransformLabel(TEXT("transform"));
+
+	return FText::Format(LOCTEXT("RemoveCurveAction_Description", "Removing {0} curve '{1}'."), FText::FromString(CurveId.CurveType == ERawCurveTrackTypes::RCT_Float ? FloatLabel : TransformLabel), FText::FromName(CurveId.InternalName.DisplayName)).ToString();
+}
+
+TUniquePtr<FChange> FSetCurveFlagsAction::ExecuteInternal(UAnimDataModel* Model, UAnimDataController* Controller)
+{
+	const FAnimCurveBase& Curve = Model->GetCurve(CurveId);
+
+	const int32 CurrentFlags = Curve.GetCurveTypeFlags();
+
+	TUniquePtr<FChange> InverseAction = MakeUnique<FSetCurveFlagsAction>(CurveId, CurrentFlags, CurveType);
+	
+	Controller->SetCurveFlags(CurveId, Flags, false);
+
+	return InverseAction;
+}
+
+FString FSetCurveFlagsAction::ToStringInternal() const
+{
+	const FString FloatLabel(TEXT("float"));
+	const FString TransformLabel(TEXT("transform"));
+	return FText::Format(LOCTEXT("SetCurveFlagsAction_Description", "Setting flags for {0} curve '{1}'."), FText::FromString(CurveType == ERawCurveTrackTypes::RCT_Float ? FloatLabel : TransformLabel), FText::FromName(CurveId.InternalName.DisplayName)).ToString();
+}
+
+TUniquePtr<FChange> FRenameCurveAction::ExecuteInternal(UAnimDataModel* Model, UAnimDataController* Controller)
+{
+	Controller->RenameCurve(CurveId, NewCurveId, false);
+	return MakeUnique<FRenameCurveAction>(NewCurveId, CurveId);
+}
+
+FString FRenameCurveAction::ToStringInternal() const
+{
+	const FString FloatLabel(TEXT("float"));
+	const FString TransformLabel(TEXT("transform"));
+	return FText::Format(LOCTEXT("RenameCurveAction_Description", "Renaming {0} curve '{1}' to '{2}'."), FText::FromString(CurveId.CurveType == ERawCurveTrackTypes::RCT_Float ? FloatLabel : TransformLabel), FText::FromName(CurveId.InternalName.DisplayName), FText::FromName(NewCurveId.InternalName.DisplayName)).ToString();
+}
+
+FScaleCurveAction::FScaleCurveAction(const FAnimationCurveIdentifier& InCurveId, float InOrigin, float InFactor, ERawCurveTrackTypes InCurveType) : CurveId(InCurveId), CurveType(InCurveType), Origin(InOrigin), Factor(InFactor)
+{
+	ensure(CurveType == ERawCurveTrackTypes::RCT_Float);
+}
+
+TUniquePtr<FChange> FScaleCurveAction::ExecuteInternal(UAnimDataModel* Model, UAnimDataController* Controller)
+{
+	const float InverseFactor = 1.0f / Factor;
+
+	TUniquePtr<FChange> InverseAction = MakeUnique<FScaleCurveAction>(CurveId, Origin, InverseFactor, CurveType);
+
+	Controller->ScaleCurve(CurveId, Origin, Factor, false);
+
+	return InverseAction;
+}
+
+FString FScaleCurveAction::ToStringInternal() const
+{
+	const FString FloatLabel(TEXT("float"));
+	const FString TransformLabel(TEXT("transform"));
+
+	return FText::Format(LOCTEXT("ScaleCurveAction_Description", "Scaling {0} curve '{1}'."), FText::FromString(CurveType == ERawCurveTrackTypes::RCT_Float ? FloatLabel : TransformLabel), FText::FromName(CurveId.InternalName.DisplayName)).ToString();
+}
+
+TUniquePtr<FChange> FAddFloatCurveAction::ExecuteInternal(UAnimDataModel* Model, UAnimDataController* Controller)
+{
+	TUniquePtr<FChange> InverseAction = MakeUnique<FRemoveCurveAction>(CurveId);
+
+	Controller->OpenBracket(LOCTEXT("AddFloatCurveAction_Description", "Adding float curve."), false);
+	Controller->AddCurve(CurveId, Flags, false);
+	Controller->SetCurveKeys(CurveId, Keys, false);
+	Controller->SetCurveColor(CurveId, Color, false);
+	Controller->CloseBracket();
+
+	return InverseAction;
+}
+
+FString FAddFloatCurveAction::ToStringInternal() const
+{
+	return FText::Format(LOCTEXT("AddFloatCurveAction_Description", "Adding float curve '{0}'."), FText::FromName(CurveId.InternalName.DisplayName)).ToString();
+}
+
+FAddTransformCurveAction::FAddTransformCurveAction(const FAnimationCurveIdentifier& InCurveId, int32 InFlags, const FTransformCurve& InTransformCurve) : CurveId(InCurveId), Flags(InFlags)
+{
+	SubCurveKeys[0] = InTransformCurve.TranslationCurve.FloatCurves[0].GetConstRefOfKeys();
+	SubCurveKeys[1] = InTransformCurve.TranslationCurve.FloatCurves[1].GetConstRefOfKeys();
+	SubCurveKeys[2] = InTransformCurve.TranslationCurve.FloatCurves[2].GetConstRefOfKeys();
+
+	SubCurveKeys[3] = InTransformCurve.RotationCurve.FloatCurves[0].GetConstRefOfKeys();
+	SubCurveKeys[4] = InTransformCurve.RotationCurve.FloatCurves[1].GetConstRefOfKeys();
+	SubCurveKeys[5] = InTransformCurve.RotationCurve.FloatCurves[2].GetConstRefOfKeys();
+
+	SubCurveKeys[6] = InTransformCurve.ScaleCurve.FloatCurves[0].GetConstRefOfKeys();
+	SubCurveKeys[7] = InTransformCurve.ScaleCurve.FloatCurves[1].GetConstRefOfKeys();
+	SubCurveKeys[8] = InTransformCurve.ScaleCurve.FloatCurves[2].GetConstRefOfKeys();
+}
+
+TUniquePtr<FChange> FAddTransformCurveAction::ExecuteInternal(UAnimDataModel* Model, UAnimDataController* Controller)
+{
+	TUniquePtr<FChange> InverseAction = MakeUnique<FRemoveCurveAction>(CurveId);
+
+	Controller->AddCurve(CurveId, Flags, false);
+
+	for (int32 SubCurveIndex = 0; SubCurveIndex < 3; ++SubCurveIndex)
+	{
+		const ETransformCurveChannel Channel = (ETransformCurveChannel)SubCurveIndex;
+		for (int32 ChannelIndex = 0; ChannelIndex < 3; ++ChannelIndex)
+		{
+			const EVectorCurveChannel Axis = (EVectorCurveChannel)ChannelIndex;
+			FAnimationCurveIdentifier TargetCurveIdentifier = CurveId;
+			UAnimationCurveIdentifierExtensions::GetTransformChildCurveIdentifier(TargetCurveIdentifier, Channel, Axis);
+			Controller->SetCurveKeys(TargetCurveIdentifier, SubCurveKeys[(SubCurveIndex * 3) + ChannelIndex], false);
+		}
+	}
+
+	return InverseAction;
+}
+
+FString FAddTransformCurveAction::ToStringInternal() const
+{
+	return FText::Format(LOCTEXT("AddFloatCurveAction_Description", "Adding transform curve '{0}'."), FText::FromName(CurveId.InternalName.DisplayName)).ToString();
+}
+
+TUniquePtr<FChange> FAddRichCurveKeyAction::ExecuteInternal(UAnimDataModel* Model, UAnimDataController* Controller)
+{
+	TUniquePtr<FChange> InverseAction = MakeUnique<FRemoveRichCurveKeyAction>(CurveId, Key.Time);
+	Controller->SetCurveKey(CurveId, Key, false);
+	return InverseAction;
+}
+
+FString FAddRichCurveKeyAction::ToStringInternal() const
+{
+	const FString FloatLabel(TEXT("float"));
+	const FString TransformLabel(TEXT("transform"));
+
+	return FText::Format(LOCTEXT("AddNamedRichCurveKeyAction_Description", "Adding key to {0} curve '{1}'."), FText::FromString(CurveId.CurveType == ERawCurveTrackTypes::RCT_Float ? FloatLabel : TransformLabel), FText::FromName(CurveId.InternalName.DisplayName)).ToString();
+}
+
+TUniquePtr<FChange> FSetRichCurveKeyAction::ExecuteInternal(UAnimDataModel* Model, UAnimDataController* Controller)
+{
+	const FRichCurve& RichCurve = Model->GetRichCurve(CurveId);
+
+	const FKeyHandle Handle = RichCurve.FindKey(Key.Time, 0.f);
+	ensure(Handle != FKeyHandle::Invalid());
+	FRichCurveKey CurrentKey = RichCurve.GetKey(Handle);
+
+	Controller->SetCurveKey(CurveId, Key, false);
+
+	return MakeUnique<FSetRichCurveKeyAction>(CurveId, CurrentKey);
+}
+
+FString FSetRichCurveKeyAction::ToStringInternal() const
+{
+	const FString FloatLabel(TEXT("float"));
+	const FString TransformLabel(TEXT("transform"));
+
+	return FText::Format(LOCTEXT("SetNamedRichCurveKeyAction_Description", "Setting key for {0} curve '{1}'."), FText::FromString(CurveId.CurveType == ERawCurveTrackTypes::RCT_Float ? FloatLabel : TransformLabel), FText::FromName(CurveId.InternalName.DisplayName)).ToString();
+}
+
+TUniquePtr<FChange> FRemoveRichCurveKeyAction::ExecuteInternal(UAnimDataModel* Model, UAnimDataController* Controller)
+{	
+	const FRichCurve& RichCurve = Model->GetRichCurve(CurveId);
+
+	const FKeyHandle Handle = RichCurve.FindKey(Time, 0.f);
+	ensure(Handle != FKeyHandle::Invalid());
+	FRichCurveKey CurrentKey = RichCurve.GetKey(Handle);
+
+	Controller->RemoveCurveKey(CurveId, Time, false);
+
+	return MakeUnique<FAddRichCurveKeyAction>(CurveId, CurrentKey);
+}
+
+FString FRemoveRichCurveKeyAction::ToStringInternal() const
+{
+	const FString FloatLabel(TEXT("float"));
+	const FString TransformLabel(TEXT("transform"));
+
+	return FText::Format(LOCTEXT("RemoveNamedRichCurveKeyAction_Description", "Removing key from {0} curve '{1}'."), FText::FromString(CurveId.CurveType == ERawCurveTrackTypes::RCT_Float ? FloatLabel : TransformLabel), FText::FromName(CurveId.InternalName.DisplayName)).ToString();
+}
+
+TUniquePtr<FChange> FSetRichCurveKeysAction::ExecuteInternal(UAnimDataModel* Model, UAnimDataController* Controller)
+{
+	const FRichCurve& RichCurve = Model->GetRichCurve(CurveId);
+
+	TUniquePtr<FChange> InverseAction = MakeUnique<FSetRichCurveKeysAction>(CurveId, RichCurve.GetConstRefOfKeys());
+
+	Controller->SetCurveKeys(CurveId, Keys, false);
+
+	return InverseAction;
+}
+
+FString FSetRichCurveKeysAction::ToStringInternal() const
+{
+	const FString FloatLabel(TEXT("float"));
+	const FString TransformLabel(TEXT("transform"));
+
+	return FText::Format(LOCTEXT("SetNamedRichCurveKeysAction_Description", "Replacing keys for {0} curve '{1}'."), FText::FromString(CurveId.CurveType == ERawCurveTrackTypes::RCT_Float ? FloatLabel : TransformLabel), FText::FromName(CurveId.InternalName.DisplayName)).ToString();
+}
+
+TUniquePtr<FChange> FSetCurveColorAction::ExecuteInternal(UAnimDataModel* Model, UAnimDataController* Controller)
+{
+	const FAnimCurveBase& AnimationCurve = Model->GetCurve(CurveId);
+	const FLinearColor CurrentColor = AnimationCurve.Color;
+
+	Controller->SetCurveColor(CurveId, Color, false);
+
+	return MakeUnique<FSetCurveColorAction>(CurveId, CurrentColor);
+}
+
+FString FSetCurveColorAction::ToStringInternal() const
+{
+	return FText::Format(LOCTEXT("SetCurveColorAction_Description", "Setting curve color '{0}'."), FText::FromName(CurveId.InternalName.DisplayName)).ToString();
+}
+
+FAddAtributeAction::FAddAtributeAction(const FAnimatedBoneAttribute& InAttribute) : AttributeId(InAttribute.Identifier)
+{
+	Keys = InAttribute.Curve.GetConstRefOfKeys();
+}
+
+TUniquePtr<FChange> FAddAtributeAction::ExecuteInternal(UAnimDataModel* Model, UAnimDataController* Controller)
+{	
+	TArray<const void*> VoidValues;
+	Algo::Transform(Keys, VoidValues, [](const FAttributeKey& Key)
+	{
+		return Key.GetValuePtr<void>();
+	});
+
+	TArray<float> Times;
+	Algo::Transform(Keys, Times, [](const FAttributeKey& Key)
+	{
+		return Key.Time;
+	});
+
+	Controller->AddAttribute(AttributeId, false);
+	Controller->SetAttributeKeys(AttributeId, MakeArrayView(Times), MakeArrayView(VoidValues), false);
+
+	return MakeUnique<FRemoveAtributeAction>(AttributeId);
+}
+
+FString FAddAtributeAction::ToStringInternal() const
+{
+	return FText::Format(LOCTEXT("AddAttributeAction_Description", "Adding attribute '{0}'."), FText::FromName(AttributeId.GetName())).ToString();
+}
+
+TUniquePtr<FChange> FRemoveAtributeAction::ExecuteInternal(UAnimDataModel* Model, UAnimDataController* Controller)
+{
+	const FAnimatedBoneAttribute& Attribute = Model->GetAttribute(AttributeId);
+	TUniquePtr<FAddAtributeAction> InverseAction = MakeUnique<FAddAtributeAction>(Attribute);
+	
+	Controller->RemoveAttribute(AttributeId, false);
+
+	return InverseAction;
+}
+
+FString FRemoveAtributeAction::ToStringInternal() const
+{
+	return FText::Format(LOCTEXT("RemoveAttributeAction_Description", "Removing attribute '{0}'."), FText::FromName(AttributeId.GetName())).ToString();
+}
+
+TUniquePtr<FChange> FAddAtributeKeyAction::ExecuteInternal(UAnimDataModel* Model, UAnimDataController* Controller)
+{
+	Controller->SetAttributeKey(AttributeId, Key.Time, Key.GetValuePtr<void>(), false);
+
+	return MakeUnique<FRemoveAtributeKeyAction>(AttributeId, Key.Time);
+}
+
+FString FAddAtributeKeyAction::ToStringInternal() const
+{
+	return FText::Format(LOCTEXT("AddAttributeKeyAction_Description", "Adding key to attribute '{0}'."), FText::FromName(AttributeId.GetName())).ToString();
+}
+
+TUniquePtr<FChange> FSetAtributeKeyAction::ExecuteInternal(UAnimDataModel* Model, UAnimDataController* Controller)
+{
+	const FAnimatedBoneAttribute& Attribute = Model->GetAttribute(AttributeId);
+
+	const FKeyHandle Handle = Attribute.Curve.FindKey(Key.Time, 0.f);
+	ensure(Handle != FKeyHandle::Invalid());
+
+	FAttributeKey CurrentKey = Attribute.Curve.GetKey(Handle);
+	Controller->SetAttributeKey(AttributeId, Key.Time, Key.GetValuePtr<void>(), false);
+
+	return MakeUnique<FSetAtributeKeyAction>(AttributeId, CurrentKey);
+}
+
+FString FSetAtributeKeyAction::ToStringInternal() const
+{
+	return FText::Format(LOCTEXT("SetAttributeKeyAction_Description", "Setting key on attribute '{0}'."), FText::FromName(AttributeId.GetName())).ToString();
+}
+
+TUniquePtr<FChange> FRemoveAtributeKeyAction::ExecuteInternal(UAnimDataModel* Model, UAnimDataController* Controller)
+{
+	const FAnimatedBoneAttribute& Attribute = Model->GetAttribute(AttributeId);
+
+	const FKeyHandle Handle = Attribute.Curve.FindKey(Time, 0.f);
+	ensure(Handle != FKeyHandle::Invalid());
+
+	FAttributeKey CurrentKey = Attribute.Curve.GetKey(Handle);
+
+	Controller->RemoveAttributeKey(AttributeId, Time, false);
+
+	return MakeUnique<FAddAtributeKeyAction>(AttributeId, CurrentKey);
+}
+
+FString FRemoveAtributeKeyAction::ToStringInternal() const
+{
+	return FText::Format(LOCTEXT("RemoveAttributeKeyAction_Description", "Removing key from attribute '{0}'."), FText::FromName(AttributeId.GetName())).ToString();
+}
+
+FSetAtributeKeysAction::FSetAtributeKeysAction(const FAnimatedBoneAttribute& InAttribute) : AttributeId(InAttribute.Identifier)
+{
+	Keys = InAttribute.Curve.GetConstRefOfKeys();
+}
+
+TUniquePtr<FChange> FSetAtributeKeysAction::ExecuteInternal(UAnimDataModel* Model, UAnimDataController* Controller)
+{
+	TUniquePtr<FSetAtributeKeysAction> InverseAction = MakeUnique<FSetAtributeKeysAction>(Model->GetAttribute(AttributeId));
+
+	TArray<const void*> VoidValues;
+	Algo::Transform(Keys, VoidValues, [](const FAttributeKey& Key)
+	{
+		return Key.GetValuePtr<void>();
+	});
+
+	TArray<float> Times;
+	Algo::Transform(Keys, Times, [](const FAttributeKey& Key)
+	{
+		return Key.Time;
+	});
+
+	Controller->SetAttributeKeys(AttributeId, MakeArrayView(Times), MakeArrayView(VoidValues), false);
+
+	return InverseAction;
+}
+
+FString FSetAtributeKeysAction::ToStringInternal() const
+{
+	return FText::Format(LOCTEXT("SetAttributeKeysAction_Description", "Replacing keys for attribute '{0}'."), FText::FromName(AttributeId.GetName())).ToString();
+}
+
+} // namespace Anim
+
+} // namespace UE
+
+#endif // WITH_EDITOR
+
+#undef LOCTEXT_NAMESPACE // "AnimDataControllerActions"
\ No newline at end of file
diff --git a/Engine/Source/Developer/AnimationDataController/Private/AnimDataControllerActions.h b/Engine/Source/Developer/AnimationDataController/Private/AnimDataControllerActions.h
new file mode 100644
--- /dev/null
+++ b/Engine/Source/Developer/AnimationDataController/Private/AnimDataControllerActions.h
@@ -0,0 +1,426 @@
+// Copyright Epic Games, Inc. All Rights Reserved.
+
+#pragma once
+
+#include "Misc/Change.h"
+
+#include "Animation/AnimData/AnimDataModel.h"
+#include "Animation/AnimTypes.h" 
+#include "Misc/FrameRate.h"
+
+class UObject;
+class UAnimDataModel;
+class UAnimDataController;
+
+#if WITH_EDITOR
+
+namespace UE {
+
+namespace Anim {
+
+/**
+* UAnimDataController instanced FChange-based objects used for storing mutations to an UAnimDataModel within the Transaction Buffer.
+* Each Action class represents an (invertable) operation mutating an UAnimDataModel object utilizing a UAnimDataController. Allowing 
+* for a more granular approach to undo/redo-ing changes while also allowing for script-based interoperability.
+*/
+class FAnimDataBaseAction : public FSwapChange
+{
+public:
+	virtual TUniquePtr<FChange> Execute(UObject* Object) final;
+	virtual ~FAnimDataBaseAction() {}
+	virtual FString ToString() const override;
+
+protected:
+	virtual TUniquePtr<FChange> ExecuteInternal(UAnimDataModel* Model, UAnimDataController* Controller) = 0;
+	virtual FString ToStringInternal() const = 0;
+};
+
+class FOpenBracketAction : public FAnimDataBaseAction
+{
+public:
+	explicit FOpenBracketAction(const FString& InDescription) : Description(InDescription) {}
+	virtual ~FOpenBracketAction() {}
+protected:
+	FOpenBracketAction() {}
+	virtual TUniquePtr<FChange> ExecuteInternal(UAnimDataModel* Model, UAnimDataController* Controller) override;
+	virtual FString ToStringInternal() const override;
+
+protected:
+	FString Description;
+};
+
+class FCloseBracketAction : public FAnimDataBaseAction
+{
+public:
+	explicit FCloseBracketAction(const FString& InDescription) : Description(InDescription) {}
+	virtual ~FCloseBracketAction() {}
+protected:
+	FCloseBracketAction() {}
+	virtual TUniquePtr<FChange> ExecuteInternal(UAnimDataModel* Model, UAnimDataController* Controller) override;
+	virtual FString ToStringInternal() const override;
+
+protected:
+	FString Description;
+};
+
+class FAddTrackAction : public FAnimDataBaseAction
+{
+public:
+	explicit FAddTrackAction(const FBoneAnimationTrack& Track, int32 TrackIndex);
+	virtual ~FAddTrackAction() {}
+protected:
+	FAddTrackAction() {}
+	virtual TUniquePtr<FChange> ExecuteInternal(UAnimDataModel* Model, UAnimDataController* Controller) override;
+	virtual FString ToStringInternal() const override;
+
+protected:
+	FName Name;
+	int32 BoneTreeIndex;
+	FRawAnimSequenceTrack Data;
+	int32 TrackIndex;
+};
+
+class FRemoveTrackAction : public FAnimDataBaseAction
+{
+public:
+	explicit FRemoveTrackAction(const FBoneAnimationTrack& Track, int32 TrackIndex);
+	virtual ~FRemoveTrackAction() {}
+protected:
+	FRemoveTrackAction() {}
+	virtual TUniquePtr<FChange> ExecuteInternal(UAnimDataModel* Model, UAnimDataController* Controller) override;
+	virtual FString ToStringInternal() const override;
+
+protected:
+	FName Name;
+	int32 TrackIndex;
+};
+
+class FSetTrackKeysAction : public FAnimDataBaseAction
+{
+public:
+	explicit FSetTrackKeysAction(const FBoneAnimationTrack& Track);
+	virtual ~FSetTrackKeysAction() {}
+protected:
+	FSetTrackKeysAction() {}
+	virtual TUniquePtr<FChange> ExecuteInternal(UAnimDataModel* Model, UAnimDataController* Controller) override;
+	virtual FString ToStringInternal() const override;
+
+protected:
+	FName Name;
+	int32 TrackIndex;
+
+	FRawAnimSequenceTrack TrackData;
+};
+
+class FResizePlayLengthAction : public FAnimDataBaseAction
+{
+public:
+	explicit FResizePlayLengthAction(const UAnimDataModel* InModel, float t0, float t1);
+	virtual ~FResizePlayLengthAction() {}
+protected:
+	FResizePlayLengthAction() {}
+	virtual TUniquePtr<FChange> ExecuteInternal(UAnimDataModel* Model, UAnimDataController* Controller) override;
+	virtual FString ToStringInternal() const override;
+
+protected:
+	float Length;
+	float T0;
+	float T1;
+};
+
+class FSetFrameRateAction : public FAnimDataBaseAction
+{
+public:
+	explicit FSetFrameRateAction(const UAnimDataModel* InModel);
+	virtual ~FSetFrameRateAction() {}
+protected:
+	FSetFrameRateAction() {}
+	virtual TUniquePtr<FChange> ExecuteInternal(UAnimDataModel* Model, UAnimDataController* Controller) override;
+	virtual FString ToStringInternal() const override;
+
+protected:
+	FFrameRate FrameRate;
+};
+
+class FAddCurveAction : public FAnimDataBaseAction
+{
+public:
+	explicit FAddCurveAction(const FAnimationCurveIdentifier& InCurveId, int32 InFlags) : CurveId(InCurveId), Flags(InFlags) {}
+
+	virtual ~FAddCurveAction() {}
+protected:
+	FAddCurveAction() {}
+	virtual TUniquePtr<FChange> ExecuteInternal(UAnimDataModel* Model, UAnimDataController* Controller) override;
+	virtual FString ToStringInternal() const override;
+
+protected:
+	FAnimationCurveIdentifier CurveId;
+	int32 Flags;
+};
+
+class FAddFloatCurveAction : public FAnimDataBaseAction
+{
+public:
+	explicit FAddFloatCurveAction(const FAnimationCurveIdentifier& InCurveId, int32 InFlags, const TArray<FRichCurveKey>& InKeys, const FLinearColor& InColor) : CurveId(InCurveId), Flags(InFlags), Keys(InKeys), Color(InColor) {}
+	virtual ~FAddFloatCurveAction() {}
+protected:
+	FAddFloatCurveAction() {}
+	virtual TUniquePtr<FChange> ExecuteInternal(UAnimDataModel* Model, UAnimDataController* Controller) override;
+	virtual FString ToStringInternal() const override;
+
+protected:
+	FAnimationCurveIdentifier CurveId;
+	int32 Flags;
+	TArray<FRichCurveKey> Keys;
+	FLinearColor Color;
+};
+
+class FAddTransformCurveAction : public FAnimDataBaseAction
+{
+public:
+	explicit FAddTransformCurveAction(const FAnimationCurveIdentifier& InCurveId, int32 InFlags, const FTransformCurve& InTransformCurve);
+	virtual ~FAddTransformCurveAction() {}
+protected:
+	FAddTransformCurveAction() {}
+	virtual TUniquePtr<FChange> ExecuteInternal(UAnimDataModel* Model, UAnimDataController* Controller) override;
+	virtual FString ToStringInternal() const override;
+
+protected:
+	FAnimationCurveIdentifier CurveId;
+	int32 Flags;
+	
+	TArray<FRichCurveKey> SubCurveKeys[9];
+};
+
+class FRemoveCurveAction : public FAnimDataBaseAction
+{
+public:
+	explicit FRemoveCurveAction(const FAnimationCurveIdentifier& InCurveId) : CurveId(InCurveId) {}
+	virtual ~FRemoveCurveAction() {}
+protected:
+	FRemoveCurveAction() {}
+	virtual TUniquePtr<FChange> ExecuteInternal(UAnimDataModel* Model, UAnimDataController* Controller) override;
+	virtual FString ToStringInternal() const override;
+
+protected:
+	FAnimationCurveIdentifier CurveId;
+};
+
+class FSetCurveFlagsAction : public FAnimDataBaseAction
+{
+public:
+	explicit FSetCurveFlagsAction(const FAnimationCurveIdentifier& InCurveId, int32 InFlags, ERawCurveTrackTypes InCurveType) : CurveId(InCurveId), Flags(InFlags), CurveType(InCurveType) {}
+	virtual ~FSetCurveFlagsAction() {}
+protected:
+	FSetCurveFlagsAction() {}
+	virtual TUniquePtr<FChange> ExecuteInternal(UAnimDataModel* Model, UAnimDataController* Controller) override;
+	virtual FString ToStringInternal() const override;
+
+protected:
+	FAnimationCurveIdentifier CurveId;
+	int32 Flags;
+	ERawCurveTrackTypes CurveType;
+};
+
+class FRenameCurveAction : public FAnimDataBaseAction
+{
+public:
+	explicit FRenameCurveAction(const FAnimationCurveIdentifier& InCurveId, const FAnimationCurveIdentifier& InNewCurveId) : CurveId(InCurveId), NewCurveId(InNewCurveId) {}
+	virtual ~FRenameCurveAction() {}
+protected:
+	FRenameCurveAction() {}
+	virtual TUniquePtr<FChange> ExecuteInternal(UAnimDataModel* Model, UAnimDataController* Controller) override;
+	virtual FString ToStringInternal() const override;
+
+protected:
+	FAnimationCurveIdentifier CurveId;
+	FAnimationCurveIdentifier NewCurveId;
+};
+
+class FScaleCurveAction : public FAnimDataBaseAction
+{
+public:
+	explicit FScaleCurveAction(const FAnimationCurveIdentifier& InCurveId, float InOrigin, float InFactor, ERawCurveTrackTypes InCurveType);
+	virtual ~FScaleCurveAction() {}
+protected:
+	FScaleCurveAction() {}
+	virtual TUniquePtr<FChange> ExecuteInternal(UAnimDataModel* Model, UAnimDataController* Controller) override;
+	virtual FString ToStringInternal() const override;
+
+protected:
+	FAnimationCurveIdentifier CurveId;
+	ERawCurveTrackTypes CurveType;
+	float Origin;
+	float Factor;
+};
+
+class FAddRichCurveKeyAction : public FAnimDataBaseAction
+{
+public:
+	explicit FAddRichCurveKeyAction(const FAnimationCurveIdentifier& InCurveId, const FRichCurveKey& InKey) : CurveId(InCurveId), Key(InKey) {}
+	virtual ~FAddRichCurveKeyAction() {}
+protected:
+	FAddRichCurveKeyAction() {}
+	virtual TUniquePtr<FChange> ExecuteInternal(UAnimDataModel* Model, UAnimDataController* Controller) override;
+	virtual FString ToStringInternal() const override;
+
+protected:
+	FAnimationCurveIdentifier CurveId;
+	ERawCurveTrackTypes CurveType;
+	FRichCurveKey Key;
+};
+
+class FSetRichCurveKeyAction : public FAnimDataBaseAction
+{
+public:
+	explicit FSetRichCurveKeyAction(const FAnimationCurveIdentifier& InCurveId, const FRichCurveKey& InKey) : CurveId(InCurveId), Key(InKey) {}
+	virtual ~FSetRichCurveKeyAction() {}
+protected:
+	FSetRichCurveKeyAction() {}
+	virtual TUniquePtr<FChange> ExecuteInternal(UAnimDataModel* Model, UAnimDataController* Controller) override;
+	virtual FString ToStringInternal() const override;
+
+protected:
+	FAnimationCurveIdentifier CurveId;
+	FRichCurveKey Key;
+};
+
+class FRemoveRichCurveKeyAction : public FAnimDataBaseAction
+{
+public:
+	explicit FRemoveRichCurveKeyAction(const FAnimationCurveIdentifier& InCurveId, const float InTime) : CurveId(InCurveId), Time(InTime) {}
+	virtual ~FRemoveRichCurveKeyAction() {}
+protected:
+	FRemoveRichCurveKeyAction() {}
+	virtual TUniquePtr<FChange> ExecuteInternal(UAnimDataModel* Model, UAnimDataController* Controller) override;
+	virtual FString ToStringInternal() const override;
+
+protected:
+	FAnimationCurveIdentifier CurveId;
+	float Time;
+};
+
+class FSetRichCurveKeysAction : public FAnimDataBaseAction
+{
+public:
+	explicit FSetRichCurveKeysAction(const FAnimationCurveIdentifier& InCurveId, const TArray<FRichCurveKey>& InKeys) : CurveId(InCurveId), Keys(InKeys) {}
+	virtual ~FSetRichCurveKeysAction() {}
+protected:
+	FSetRichCurveKeysAction() {}
+	virtual TUniquePtr<FChange> ExecuteInternal(UAnimDataModel* Model, UAnimDataController* Controller) override;
+	virtual FString ToStringInternal() const override;
+
+protected:
+	FAnimationCurveIdentifier CurveId;
+	TArray<FRichCurveKey> Keys;
+};
+
+class FSetCurveColorAction : public FAnimDataBaseAction
+{
+public:
+	explicit FSetCurveColorAction(const FAnimationCurveIdentifier& InCurveId, const FLinearColor& InColor) : CurveId(InCurveId), Color(InColor) {}
+	virtual ~FSetCurveColorAction() {}
+protected:
+	FSetCurveColorAction() {}
+	virtual TUniquePtr<FChange> ExecuteInternal(UAnimDataModel* Model, UAnimDataController* Controller) override;
+	virtual FString ToStringInternal() const override;
+
+protected:
+	FAnimationCurveIdentifier CurveId;
+	FLinearColor Color;
+};
+
+class FAddAtributeAction : public FAnimDataBaseAction
+{
+public:
+	explicit FAddAtributeAction(const FAnimatedBoneAttribute& InAttribute);
+	virtual ~FAddAtributeAction() {}
+protected:
+	FAddAtributeAction() {}
+	virtual TUniquePtr<FChange> ExecuteInternal(UAnimDataModel* Model, UAnimDataController* Controller) override;
+	virtual FString ToStringInternal() const override;
+
+protected:
+	FAnimationAttributeIdentifier AttributeId;
+	TArray<FAttributeKey> Keys;
+};
+
+class FRemoveAtributeAction : public FAnimDataBaseAction
+{
+public:
+	explicit FRemoveAtributeAction(const FAnimationAttributeIdentifier& InAttributeId) : AttributeId(InAttributeId) {}
+	virtual ~FRemoveAtributeAction() {}
+protected:
+	FRemoveAtributeAction() {}
+	virtual TUniquePtr<FChange> ExecuteInternal(UAnimDataModel* Model, UAnimDataController* Controller) override;
+	virtual FString ToStringInternal() const override;
+
+protected:
+	FAnimationAttributeIdentifier AttributeId;
+};
+
+class FAddAtributeKeyAction : public FAnimDataBaseAction
+{
+public:
+	explicit FAddAtributeKeyAction(const FAnimationAttributeIdentifier& InAttributeId, const FAttributeKey& InKey) : AttributeId(InAttributeId), Key(InKey) {}
+	virtual ~FAddAtributeKeyAction() {}
+protected:
+	FAddAtributeKeyAction() {}
+	virtual TUniquePtr<FChange> ExecuteInternal(UAnimDataModel* Model, UAnimDataController* Controller) override;
+	virtual FString ToStringInternal() const override;
+
+protected:
+	FAnimationAttributeIdentifier AttributeId;
+	FAttributeKey Key;
+};
+
+class FSetAtributeKeyAction : public FAnimDataBaseAction
+{
+public:
+	explicit FSetAtributeKeyAction(const FAnimationAttributeIdentifier& InAttributeId, const FAttributeKey& InKey) : AttributeId(InAttributeId), Key(InKey) {}
+	virtual ~FSetAtributeKeyAction() {}
+protected:
+	FSetAtributeKeyAction() {}
+	virtual TUniquePtr<FChange> ExecuteInternal(UAnimDataModel* Model, UAnimDataController* Controller) override;
+	virtual FString ToStringInternal() const override;
+
+protected:
+	FAnimationAttributeIdentifier AttributeId;
+	FAttributeKey Key;
+};
+
+class FRemoveAtributeKeyAction : public FAnimDataBaseAction
+{
+public:
+	explicit FRemoveAtributeKeyAction(const FAnimationAttributeIdentifier& InAttributeId, float InTime) : AttributeId(InAttributeId), Time(InTime) {}
+	virtual ~FRemoveAtributeKeyAction() {}
+protected:
+	FRemoveAtributeKeyAction() {}
+	virtual TUniquePtr<FChange> ExecuteInternal(UAnimDataModel* Model, UAnimDataController* Controller) override;
+	virtual FString ToStringInternal() const override;
+
+protected:
+	FAnimationAttributeIdentifier AttributeId;
+	float Time;
+};
+
+class FSetAtributeKeysAction : public FAnimDataBaseAction
+{
+public:
+	explicit FSetAtributeKeysAction(const FAnimatedBoneAttribute& InAttribute);
+	virtual ~FSetAtributeKeysAction() {}
+protected:
+	FSetAtributeKeysAction() {}
+	virtual TUniquePtr<FChange> ExecuteInternal(UAnimDataModel* Model, UAnimDataController* Controller) override;
+	virtual FString ToStringInternal() const override;
+
+protected:
+	FAnimationAttributeIdentifier AttributeId;
+	TArray<FAttributeKey> Keys;
+};
+
+} // namespace Anim
+
+} // namespace UE
+
+#endif // WITH_EDITOR
\ No newline at end of file
diff --git a/Engine/Source/Developer/AnimationDataController/Private/AnimationDataControllerModule.cpp b/Engine/Source/Developer/AnimationDataController/Private/AnimationDataControllerModule.cpp
new file mode 100644
--- /dev/null
+++ b/Engine/Source/Developer/AnimationDataController/Private/AnimationDataControllerModule.cpp
@@ -0,0 +1,17 @@
+// Copyright Epic Games, Inc. All Rights Reserved.
+
+#include "IAnimationDataControllerModule.h"
+#include "AnimDataController.h"
+#include "UObject/StrongObjectPtr.h"
+
+class FAnimationDataControllerModule : public IAnimationDataControllerModule
+{
+	virtual TScriptInterface<IAnimationDataController> GetController(UObject* Outer = GetTransientPackage()) override
+	{
+		check(Outer);
+		UAnimDataController* AnimController = NewObject<UAnimDataController>(Outer);
+		return TScriptInterface<IAnimationDataController>(AnimController);
+	}
+};
+
+IMPLEMENT_MODULE(FAnimationDataControllerModule, AnimationDataController);
\ No newline at end of file
diff --git a/Engine/Source/Developer/AnimationDataController/Public/AnimDataController.h b/Engine/Source/Developer/AnimationDataController/Public/AnimDataController.h
new file mode 100644
--- /dev/null
+++ b/Engine/Source/Developer/AnimationDataController/Public/AnimDataController.h
@@ -0,0 +1,134 @@
+// Copyright Epic Games, Inc. All Rights Reserved.
+
+#pragma once
+
+#include "CoreTypes.h"
+#include "UObject/ObjectMacros.h"
+#include "Animation/AnimData/AnimDataModel.h"
+#include "Animation/AnimCurveTypes.h"
+#include "Algo/Transform.h"
+#include "Animation/AnimData/IAnimationDataController.h"
+
+#if WITH_EDITOR
+#include "ChangeTransactor.h"
+#endif // WITH_EDITOR
+
+#include "AnimDataController.generated.h"
+
+#if WITH_EDITOR
+
+#endif // WITH_EDITOR
+
+struct FAnimationCurveIdentifier;
+struct FAnimationAttributeIdentifier;
+
+namespace UE {
+namespace Anim {
+	class FOpenBracketAction;
+	class FCloseBracketAction;
+	
+	static const int32 DefaultCurveFlags = EAnimAssetCurveFlags::AACF_Editable;
+}}
+
+     UCLASS(BlueprintType)
+class ANIMATIONDATACONTROLLER_API UAnimDataController : public UObject, public IAnimationDataController
+{
+	GENERATED_BODY()
+
+public:
+	UAnimDataController() 
+#if WITH_EDITOR
+	: BracketDepth(0) 
+#endif // WITH_EDITOR
+	{}
+
+#if WITH_EDITOR
+	/** Begin IAnimationDataController overrides */
+	virtual void SetModel(UAnimDataModel* InModel)  override;
+    virtual UAnimDataModel* GetModel() override { return Model; }
+	virtual const UAnimDataModel* const GetModel() const override { return Model; }
+	virtual void OpenBracket(const FText& InTitle, bool bShouldTransact = true) override;
+	virtual void CloseBracket(bool bShouldTransact = true) override;
+	virtual void SetPlayLength(float Length, bool bShouldTransact = true) override;
+	virtual void ResizePlayLength(float NewLength, float T0, float T1, bool bShouldTransact = true) override;
+	virtual void Resize(float Length, float T0, float T1, bool bShouldTransact = true) override;
+	virtual void SetFrameRate(FFrameRate FrameRate, bool bShouldTransact = true) override;
+	virtual int32 AddBoneTrack(FName BoneName, bool bShouldTransact = true) override;
+	virtual int32 InsertBoneTrack(FName BoneName, int32 DesiredIndex, bool bShouldTransact = true) override;
+	virtual bool RemoveBoneTrack(FName BoneName, bool bShouldTransact = true) override;
+	virtual void RemoveAllBoneTracks(bool bShouldTransact = true) override;
+	virtual bool SetBoneTrackKeys(FName BoneName, const TArray<FVector>& PositionalKeys, const TArray<FQuat>& RotationalKeys, const TArray<FVector>& ScalingKeys, bool bShouldTransact = true) override;	
+	virtual bool AddCurve(const FAnimationCurveIdentifier& CurveId, int32 CurveFlags = 0x00000004, bool bShouldTransact = true) override;
+	virtual bool DuplicateCurve(const FAnimationCurveIdentifier& CopyCurveId, const FAnimationCurveIdentifier& NewCurveId, bool bShouldTransact = true) override;
+	virtual bool RemoveCurve(const FAnimationCurveIdentifier& CurveId, bool bShouldTransact = true) override;
+	virtual void RemoveAllCurvesOfType(ERawCurveTrackTypes SupportedCurveType, bool bShouldTransact = true) override;
+	virtual bool SetCurveFlag(const FAnimationCurveIdentifier& CurveId, EAnimAssetCurveFlags Flag, bool bState = true, bool bShouldTransact = true) override;
+	virtual bool SetCurveFlags(const FAnimationCurveIdentifier& CurveId, int32 Flags, bool bShouldTransact = true) override;
+	virtual bool SetTransformCurveKeys(const FAnimationCurveIdentifier& CurveId, const TArray<FTransform>& TransformValues, const TArray<float>& TimeKeys, bool bShouldTransact = true) override;	
+	virtual bool SetTransformCurveKey(const FAnimationCurveIdentifier& CurveId, float Time, const FTransform& Value, bool bShouldTransact = true) override;
+	virtual bool RemoveTransformCurveKey(const FAnimationCurveIdentifier& CurveId, float Time, bool bShouldTransact = true) override;
+	virtual bool RenameCurve(const FAnimationCurveIdentifier& CurveToRenameId, const FAnimationCurveIdentifier& NewCurveId, bool bShouldTransact = true) override;
+	virtual bool SetCurveColor(const FAnimationCurveIdentifier& CurveId, FLinearColor Color, bool bShouldTransact = true) override;	
+	virtual bool ScaleCurve(const FAnimationCurveIdentifier& CurveId, float Origin, float Factor, bool bShouldTransact = true) override;
+	virtual bool SetCurveKey(const FAnimationCurveIdentifier& CurveId, const FRichCurveKey& Key, bool bShouldTransact = true) override;	
+	virtual bool RemoveCurveKey(const FAnimationCurveIdentifier& CurveId, float Time, bool bShouldTransact = true) override;
+	virtual bool SetCurveKeys(const FAnimationCurveIdentifier& CurveId, const TArray<FRichCurveKey>& CurveKeys, bool bShouldTransact = true) override;
+	virtual void UpdateCurveNamesFromSkeleton(const USkeleton* Skeleton, ERawCurveTrackTypes SupportedCurveType, bool bShouldTransact = true) override;
+	virtual void FindOrAddCurveNamesOnSkeleton(USkeleton* Skeleton, ERawCurveTrackTypes SupportedCurveType, bool bShouldTransact = true) override;
+	virtual bool RemoveBoneTracksMissingFromSkeleton(const USkeleton* Skeleton, bool bShouldTransact = true) override;
+	virtual void NotifyPopulated() override;
+	virtual void ResetModel(bool bShouldTransact = true) override;
+	
+protected:
+	virtual void NotifyBracketOpen() override;
+	virtual void NotifyBracketClosed() override;
+	/** End IAnimationDataController overrides */
+	
+private:
+	/** Returns whether or not the supplied curve type is supported by the controller functionality */
+	const bool IsSupportedCurveType(ERawCurveTrackTypes CurveType) const;
+	/** Returns the string representation of the provided curve enum type value */
+	FString GetCurveTypeValueName(ERawCurveTrackTypes InType) const;
+	
+	/** Resizes the curve data stored on the model according to the provided new length and time at which to insert or remove time */
+	void ResizeCurves(float NewLength, bool bInserted, float T0, float T1, bool bShouldTransact = true);
+
+	/** Ensures that a valid model is currently targeted */
+	void ValidateModel() const;
+
+	/** Verifies whether or not the Model's outer object is (or is derived from) the specified UClass */
+	bool CheckOuterClass(UClass* InClass) const;
+
+	/** Helper functionality to output script-based warnings and errors */
+	void ReportWarning(const FText& InMessage) const;
+	void ReportError(const FText& InMessage) const;
+
+	template <typename FmtType, typename... Types>
+    void ReportWarningf(const FmtType& Fmt, Types... Args) const
+	{	
+		ReportWarning(FText::Format(Fmt, Args...));
+	}
+
+	template <typename FmtType, typename... Types>
+    void ReportErrorf(const FmtType& Fmt, Types... Args) const
+	{
+		ReportError(FText::Format(Fmt, Args...));
+	}
+#endif // WITH_EDITOR
+
+private: 
+#if WITH_EDITOR
+	int32 BracketDepth;
+
+	UE::FChangeTransactor ChangeTransactor;	
+#endif // WITH_EDITOR
+
+#if WITH_EDITORONLY_DATA
+	UPROPERTY(transient)
+	TObjectPtr<UAnimDataModel> Model;
+#endif // WITH_EDITORONLY_DATA
+
+	friend class FAnimDataControllerTestBase;
+	friend UE::Anim::FOpenBracketAction;
+	friend UE::Anim::FCloseBracketAction;
+};
diff --git a/Engine/Source/Developer/AnimationDataController/Public/IAnimationDataControllerModule.h b/Engine/Source/Developer/AnimationDataController/Public/IAnimationDataControllerModule.h
new file mode 100644
--- /dev/null
+++ b/Engine/Source/Developer/AnimationDataController/Public/IAnimationDataControllerModule.h
@@ -0,0 +1,16 @@
+// Copyright Epic Games, Inc. All Rights Reserved.
+
+#pragma once
+
+#include "Animation/AnimData/IAnimationDataController.h"
+#include "Modules/ModuleInterface.h"
+
+class UAnimSequence;
+class UAnimSequenceBase;
+
+class IAnimationDataControllerModule : public IModuleInterface
+{
+public:
+	/** Returns UAnimationDataController instance, with optional outer, wrapped in a TScriptInterface of its implement IAnimationDataController interface */
+	virtual TScriptInterface<IAnimationDataController> GetController(UObject* Outer = GetTransientPackage()) = 0;
+};
\ No newline at end of file
diff --git a/Engine/Source/Developer/AssetTools/Private/AssetTypeActions/AssetTypeActions_Skeleton.cpp b/Engine/Source/Developer/AssetTools/Private/AssetTypeActions/AssetTypeActions_Skeleton.cpp
--- a/Engine/Source/Developer/AssetTools/Private/AssetTypeActions/AssetTypeActions_Skeleton.cpp
+++ b/Engine/Source/Developer/AssetTools/Private/AssetTypeActions/AssetTypeActions_Skeleton.cpp
@@ -905,18 +905,18 @@ void FAssetTypeActions_Skeleton::RetargetSkeleton(TArray<FAssetToRemapSkeleton>&
 			{
 				if ( UAnimationAsset* AnimAsset = Cast<UAnimationAsset>(Asset) )
 				{
-					TUniquePtr<UAnimDataController::FScopedBracket> ScopedBracket;
+					TUniquePtr<IAnimationDataController::FScopedBracket> ScopedBracket;
 
 					if (UAnimSequenceBase* SequenceBase = Cast<UAnimSequenceBase>(AnimAsset))
 					{
-						UAnimDataController* Controller = SequenceBase->GetController();
-						ScopedBracket = MakeUnique<UAnimDataController::FScopedBracket>(Controller, LOCTEXT("FAssetTypeActions_Skeleton_Bracket", "Retargeting Skeleton"));
+						IAnimationDataController& Controller = SequenceBase->GetController();
+						ScopedBracket = MakeUnique<IAnimationDataController::FScopedBracket>(Controller, LOCTEXT("FAssetTypeActions_Skeleton_Bracket", "Retargeting Skeleton"));
 
-						Controller->FindOrAddCurveNamesOnSkeleton(NewSkeleton, ERawCurveTrackTypes::RCT_Float);
+						Controller.FindOrAddCurveNamesOnSkeleton(NewSkeleton, ERawCurveTrackTypes::RCT_Float);
 						
 						if (UAnimSequence* Sequence = Cast<UAnimSequence>(SequenceBase))
 						{
-							Controller->FindOrAddCurveNamesOnSkeleton(NewSkeleton, ERawCurveTrackTypes::RCT_Transform);
+							Controller.FindOrAddCurveNamesOnSkeleton(NewSkeleton, ERawCurveTrackTypes::RCT_Transform);
 						}
 					}
 					
diff --git a/Engine/Source/Editor/AnimationEditor/Private/AnimationEditor.cpp b/Engine/Source/Editor/AnimationEditor/Private/AnimationEditor.cpp
--- a/Engine/Source/Editor/AnimationEditor/Private/AnimationEditor.cpp
+++ b/Engine/Source/Editor/AnimationEditor/Private/AnimationEditor.cpp
@@ -663,11 +663,11 @@ void FAnimationEditor::OnRemoveBoneTrack()
 		UAnimSequence* AnimSequence = Cast<UAnimSequence>(AnimationAsset);
 		if (AnimSequence)
 		{
-			UAnimDataController* Controller = AnimSequence->GetController();
-			UAnimDataController::FScopedBracket ScopedBracket(Controller, LOCTEXT("OnRemoveBoneTrack_Bracket", "Removing all Bone Animation and Transform Curve Tracks"));
+			IAnimationDataController& Controller = AnimSequence->GetController();
+			IAnimationDataController::FScopedBracket ScopedBracket(Controller, LOCTEXT("OnRemoveBoneTrack_Bracket", "Removing all Bone Animation and Transform Curve Tracks"));
 
-			Controller->RemoveAllBoneTracks();
-			Controller->RemoveAllCurvesOfType(ERawCurveTrackTypes::RCT_Transform);			
+			Controller.RemoveAllBoneTracks();
+			Controller.RemoveAllCurvesOfType(ERawCurveTrackTypes::RCT_Transform);			
 		}
 	}
 }
diff --git a/Engine/Source/Editor/AnimationModifiers/Private/AnimationBlueprintLibrary.cpp b/Engine/Source/Editor/AnimationModifiers/Private/AnimationBlueprintLibrary.cpp
--- a/Engine/Source/Editor/AnimationModifiers/Private/AnimationBlueprintLibrary.cpp
+++ b/Engine/Source/Editor/AnimationModifiers/Private/AnimationBlueprintLibrary.cpp
@@ -17,7 +17,6 @@
 #include "AnimationRuntime.h"
 #include "Animation/AnimData/AnimDataModel.h"
 #include "Animation/AnimSequenceHelpers.h"
-#include "Animation/AnimData/AnimDataController.h"
 
 #define LOCTEXT_NAMESPACE "AnimationBlueprintLibrary"
 
@@ -1265,10 +1264,10 @@ void UAnimationBlueprintLibrary::RemoveAllCurveData(UAnimSequence* AnimationSequ
 {
 	if (AnimationSequence)
 	{
-		UAnimDataController* Controller = AnimationSequence->GetController();
+		IAnimationDataController& Controller = AnimationSequence->GetController();
 
-		Controller->RemoveAllCurvesOfType(ERawCurveTrackTypes::RCT_Float);
-		Controller->RemoveAllCurvesOfType(ERawCurveTrackTypes::RCT_Transform);
+		Controller.RemoveAllCurvesOfType(ERawCurveTrackTypes::RCT_Float);
+		Controller.RemoveAllCurvesOfType(ERawCurveTrackTypes::RCT_Transform);
 	}
 	else
 	{
@@ -1393,27 +1392,27 @@ void UAnimationBlueprintLibrary::AddVectorCurveKeys(UAnimSequence* AnimationSequ
 	}
 }
 
-static void SetControllerCurveKeys(UAnimDataController* Controller, FSmartName SmartName, const TArray<float>& Times, const TArray<float>& Values)
+static void SetControllerCurveKeys(IAnimationDataController& Controller, FSmartName SmartName, const TArray<float>& Times, const TArray<float>& Values)
 {
 	const int32 NumKeys = Values.Num();
 	const FAnimationCurveIdentifier CurveId(SmartName, ERawCurveTrackTypes::RCT_Float);
 	for (int32 KeyIndex = 0; KeyIndex < NumKeys; ++KeyIndex)
 	{
-		Controller->SetCurveKey(CurveId, { Times[KeyIndex], Values[KeyIndex] });
+		Controller.SetCurveKey(CurveId, { Times[KeyIndex], Values[KeyIndex] });
 	}
 }
 
-static void SetControllerCurveKeys(UAnimDataController* Controller, FSmartName SmartName, const TArray<float>& Times, const TArray<FTransform>& Values)
+static void SetControllerCurveKeys(IAnimationDataController& Controller, FSmartName SmartName, const TArray<float>& Times, const TArray<FTransform>& Values)
 {
 	const int32 NumKeys = Values.Num();
 	const FAnimationCurveIdentifier CurveId(SmartName, ERawCurveTrackTypes::RCT_Transform);
 	for (int32 KeyIndex = 0; KeyIndex < NumKeys; ++KeyIndex)
 	{
-		Controller->SetTransformCurveKey(CurveId, Times[KeyIndex], Values[KeyIndex]);
+		Controller.SetTransformCurveKey(CurveId, Times[KeyIndex], Values[KeyIndex]);
 	}
 }
 	
-static void SetControllerCurveKeys(UAnimDataController* Controller, FSmartName CurveId, const TArray<float>& Times, const TArray<FVector>& Values)
+static void SetControllerCurveKeys(IAnimationDataController& Controller, FSmartName CurveId, const TArray<float>& Times, const TArray<FVector>& Values)
 {
 	ensure(false);
 }
@@ -1435,7 +1434,7 @@ void UAnimationBlueprintLibrary::AddCurveKeysInternal(UAnimSequence* AnimationSe
 		const CurveClass* Curve = static_cast<const CurveClass*>(AnimationSequence->GetDataModel()->FindCurve(CurveId));
 		if (Curve)
 		{
-			UAnimDataController* Controller = AnimationSequence->GetController();
+			IAnimationDataController& Controller = AnimationSequence->GetController();
 			SetControllerCurveKeys(Controller, CurveSmartName, Times, KeyData);
 		}
 	}
@@ -1454,8 +1453,8 @@ bool UAnimationBlueprintLibrary::AddCurveInternal(UAnimSequence* AnimationSequen
 	const FAnimCurveBase* ExistingCurve = AnimationSequence->GetDataModel()->FindCurve(CurveId);
 	if (ExistingCurve == nullptr)
 	{
-		UAnimDataController* Controller = AnimationSequence->GetController();
-		Controller->AddCurve(CurveId, CurveFlags);
+		IAnimationDataController& Controller = AnimationSequence->GetController();
+		Controller.AddCurve(CurveId, CurveFlags);
 		bCurveAdded = AnimationSequence->GetDataModel()->FindCurve(CurveId) != nullptr;	
 	}
 	else
@@ -1478,10 +1477,10 @@ bool UAnimationBlueprintLibrary::RemoveCurveInternal(UAnimSequence* AnimationSeq
 		checkf(Skeleton != nullptr, TEXT("Invalid Skeleton ptr"));
 		if (Skeleton->GetSmartNameByUID(ContainerName, UID, SmartCurveName))
 		{
-			UAnimDataController* Controller = AnimationSequence->GetController();
+			IAnimationDataController& Controller = AnimationSequence->GetController();
 
 			const FAnimationCurveIdentifier CurveId(SmartCurveName, ContainerName == USkeleton::AnimTrackCurveMappingName ? ERawCurveTrackTypes::RCT_Transform : ERawCurveTrackTypes::RCT_Float);
-			bRemoved = Controller->RemoveCurve(CurveId);
+			bRemoved = Controller.RemoveCurve(CurveId);
 
 			if (bRemoveNameFromSkeleton)
 			{
@@ -1583,9 +1582,9 @@ void UAnimationBlueprintLibrary::GetTransformationKeys(UAnimSequence* AnimationS
 
 void UAnimationBlueprintLibrary::CopyAnimationCurveNamesToSkeleton(USkeleton* OldSkeleton, USkeleton* NewSkeleton, UAnimSequenceBase* SequenceBase, ERawCurveTrackTypes CurveType)
 {
-	UAnimDataController* Controller = SequenceBase->GetController();
-	Controller->UpdateCurveNamesFromSkeleton(OldSkeleton, CurveType);
-	Controller->FindOrAddCurveNamesOnSkeleton(NewSkeleton, CurveType);	
+	IAnimationDataController& Controller = SequenceBase->GetController();
+	Controller.UpdateCurveNamesFromSkeleton(OldSkeleton, CurveType);
+	Controller.FindOrAddCurveNamesOnSkeleton(NewSkeleton, CurveType);	
 }
 
 template <typename DataType, typename CurveClass, ERawCurveTrackTypes CurveType>
@@ -2220,12 +2219,12 @@ void UAnimationBlueprintLibrary::RemoveBoneAnimation(UAnimSequence* AnimationSeq
 				}
 			}
 
-			UAnimDataController* Controller = AnimationSequence->GetController();
+			IAnimationDataController& Controller = AnimationSequence->GetController();
 
-			UAnimDataController::FScopedBracket ScopedBracket(Controller, LOCTEXT("RemoveBoneAnimation_Bracket", "Removing Bone Animation Track"));
+			IAnimationDataController::FScopedBracket ScopedBracket(Controller, LOCTEXT("RemoveBoneAnimation_Bracket", "Removing Bone Animation Track"));
 			for (const FName& TrackName : TracksToRemove)
 			{
-				Controller->RemoveBoneTrack(TrackName);
+				Controller.RemoveBoneTrack(TrackName);
 			}
 		}
 		else
@@ -2240,10 +2239,10 @@ void UAnimationBlueprintLibrary::RemoveAllBoneAnimation(UAnimSequence* Animation
 {
 	if (AnimationSequence)
 	{
-		UAnimDataController* Controller = AnimationSequence->GetController();
-		UAnimDataController::FScopedBracket ScopedBracket(Controller, LOCTEXT("RemoveAllBoneAnimation_Bracket", "Removing all Bone Animation and Transform Curve Tracks"));
-		Controller->RemoveAllBoneTracks();
-		Controller->RemoveAllCurvesOfType(ERawCurveTrackTypes::RCT_Transform);		
+		IAnimationDataController& Controller = AnimationSequence->GetController();
+		IAnimationDataController::FScopedBracket ScopedBracket(Controller, LOCTEXT("RemoveAllBoneAnimation_Bracket", "Removing all Bone Animation and Transform Curve Tracks"));
+		Controller.RemoveAllBoneTracks();
+		Controller.RemoveAllCurvesOfType(ERawCurveTrackTypes::RCT_Transform);		
 	}
 }
 
diff --git a/Engine/Source/Editor/AnimationModifiers/Private/AnimationModifier.cpp b/Engine/Source/Editor/AnimationModifiers/Private/AnimationModifier.cpp
--- a/Engine/Source/Editor/AnimationModifiers/Private/AnimationModifier.cpp
+++ b/Engine/Source/Editor/AnimationModifiers/Private/AnimationModifier.cpp
@@ -52,10 +52,10 @@ void UAnimationModifier::ApplyToAnimationSequence(class UAnimSequence* InAnimati
 		PreviouslyAppliedModifier->OnRevert(CurrentAnimSequence);
 	}
 
-	UAnimDataController* Controller = CurrentAnimSequence->GetController();
+	IAnimationDataController& Controller = CurrentAnimSequence->GetController();
 
 	{
-		UAnimDataController::FScopedBracket ScopedBracket(Controller, LOCTEXT("ApplyModifierBracket", "Applying Animation Modifier"));
+		IAnimationDataController::FScopedBracket ScopedBracket(Controller, LOCTEXT("ApplyModifierBracket", "Applying Animation Modifier"));
 		/** Reverting and applying, populates the log with possible warnings and or errors to notify the user about */
 		OnApply(CurrentAnimSequence);
 	}
@@ -136,10 +136,10 @@ void UAnimationModifier::RevertFromAnimationSequence(class UAnimSequence* InAnim
 
 		PreviouslyAppliedModifier->Modify();
 
-		UAnimDataController* Controller = CurrentAnimSequence->GetController();
+		IAnimationDataController& Controller = CurrentAnimSequence->GetController();
 
 		{
-			UAnimDataController::FScopedBracket ScopedBracket(Controller, LOCTEXT("RevertModifierBracket", "Reverting Animation Modifier"));
+			IAnimationDataController::FScopedBracket ScopedBracket(Controller, LOCTEXT("RevertModifierBracket", "Reverting Animation Modifier"));
 			PreviouslyAppliedModifier->OnRevert(CurrentAnimSequence);
 		}
 
diff --git a/Engine/Source/Editor/Persona/Private/AnimTimeline/AnimModel_AnimSequenceBase.cpp b/Engine/Source/Editor/Persona/Private/AnimTimeline/AnimModel_AnimSequenceBase.cpp
--- a/Engine/Source/Editor/Persona/Private/AnimTimeline/AnimModel_AnimSequenceBase.cpp
+++ b/Engine/Source/Editor/Persona/Private/AnimTimeline/AnimModel_AnimSequenceBase.cpp
@@ -24,7 +24,6 @@
 #include "AnimPreviewInstance.h"
 #include "ScopedTransaction.h"
 
-#include "Animation/AnimData/AnimDataController.h"
 
 #define LOCTEXT_NAMESPACE "FAnimModel_AnimSequence"
 
@@ -407,8 +406,8 @@ bool FAnimModel_AnimSequenceBase::CanEditSelectedCurves() const
 
 void FAnimModel_AnimSequenceBase::RemoveSelectedCurves()
 {
-	UAnimDataController* Controller = AnimSequenceBase->GetController();
-	Controller->OpenBracket(LOCTEXT("CurvePanel_RemoveCurves", "Remove Curves"));
+	IAnimationDataController& Controller = AnimSequenceBase->GetController();
+	Controller.OpenBracket(LOCTEXT("CurvePanel_RemoveCurves", "Remove Curves"));
 
 	bool bDeletedCurve = false;
 
@@ -433,7 +432,7 @@ void FAnimModel_AnimSequenceBase::RemoveSelectedCurves()
 					IAnimationEditor::FCurveEditInfo EditInfo(Name, Type, CurveEditIndex);
 					OnStopEditingCurves.ExecuteIfBound(TArray<IAnimationEditor::FCurveEditInfo>({ EditInfo }));
 
-					Controller->RemoveCurve(FloatCurveId);
+					Controller.RemoveCurve(FloatCurveId);
 					bDeletedCurve = true;
 				}
 			}
@@ -466,14 +465,14 @@ void FAnimModel_AnimSequenceBase::RemoveSelectedCurves()
 
 					OnStopEditingCurves.ExecuteIfBound(CurveEditInfo);
 
-					Controller->RemoveCurve(TransformCurveId);
+					Controller.RemoveCurve(TransformCurveId);
 					bDeletedCurve = true;
 				}
 			}	
 		}
 	}
 
-	Controller->CloseBracket();
+	Controller.CloseBracket();
 
 	if(bDeletedCurve)
 	{
diff --git a/Engine/Source/Editor/Persona/Private/AnimTimeline/AnimTimelineTrack_Curves.cpp b/Engine/Source/Editor/Persona/Private/AnimTimeline/AnimTimelineTrack_Curves.cpp
--- a/Engine/Source/Editor/Persona/Private/AnimTimeline/AnimTimelineTrack_Curves.cpp
+++ b/Engine/Source/Editor/Persona/Private/AnimTimeline/AnimTimelineTrack_Curves.cpp
@@ -14,7 +14,6 @@
 #include "Animation/AnimMontage.h"
 #include "SAnimOutlinerItem.h"
 #include "Preferences/PersonaOptions.h"
-#include "Animation/AnimData/AnimDataController.h"
 #include "SListViewSelectorDropdownMenu.h"
 
 #define LOCTEXT_NAMESPACE "FAnimTimelineTrack_Notifies"
@@ -82,8 +81,8 @@ TSharedRef<SWidget> FAnimTimelineTrack_Curves::GenerateContainerWidgetForOutline
 void FAnimTimelineTrack_Curves::DeleteAllCurves()
 {
 	UAnimSequenceBase* AnimSequenceBase = GetModel()->GetAnimSequenceBase();
-	UAnimDataController* Controller = AnimSequenceBase->GetController();
-	Controller->RemoveAllCurvesOfType(ERawCurveTrackTypes::RCT_Float);
+	IAnimationDataController& Controller = AnimSequenceBase->GetController();
+	Controller.RemoveAllCurvesOfType(ERawCurveTrackTypes::RCT_Float);
 }
 
 TSharedRef<SWidget> FAnimTimelineTrack_Curves::BuildCurvesSubMenu()
@@ -413,12 +412,12 @@ void FAnimTimelineTrack_Curves::AddMetadataEntry(USkeleton::AnimCurveUID Uid)
 	UAnimSequenceBase* AnimSequenceBase = GetModel()->GetAnimSequenceBase();
 	ensureAlways(AnimSequenceBase->GetSkeleton()->GetSmartNameByUID(USkeleton::AnimCurveMappingName, Uid, NewName));
 
-	UAnimDataController* Controller = AnimSequenceBase->GetController();
-	UAnimDataController::FScopedBracket ScopedBracket(Controller, LOCTEXT("AddCurveMetadata", "Add Curve Metadata"));
+	IAnimationDataController& Controller = AnimSequenceBase->GetController();
+	IAnimationDataController::FScopedBracket ScopedBracket(Controller, LOCTEXT("AddCurveMetadata", "Add Curve Metadata"));
 
 	const FAnimationCurveIdentifier MetadataCurveId(NewName, ERawCurveTrackTypes::RCT_Float);
-	Controller->AddCurve(MetadataCurveId, AACF_Metadata);
-	Controller->SetCurveKeys(MetadataCurveId, { FRichCurveKey(0.f, 1.f) });	
+	Controller.AddCurve(MetadataCurveId, AACF_Metadata);
+	Controller.SetCurveKeys(MetadataCurveId, { FRichCurveKey(0.f, 1.f) });	
 }
 
 void FAnimTimelineTrack_Curves::CreateNewMetadataEntryClicked()
@@ -507,9 +506,9 @@ void FAnimTimelineTrack_Curves::AddVariableCurve(USkeleton::AnimCurveUID CurveUi
 	FSmartName NewName;
 	ensureAlways(Skeleton->GetSmartNameByUID(USkeleton::AnimCurveMappingName, CurveUid, NewName));
 
-	UAnimDataController* Controller = AnimSequenceBase->GetController();
+	IAnimationDataController& Controller = AnimSequenceBase->GetController();
 	const FAnimationCurveIdentifier FloatCurveId(NewName, ERawCurveTrackTypes::RCT_Float);
-	Controller->AddCurve(FloatCurveId);
+	Controller.AddCurve(FloatCurveId);
 }
 
 void FAnimTimelineTrack_Curves::HandleShowCurvePoints()
diff --git a/Engine/Source/Editor/Persona/Private/AnimTimeline/AnimTimelineTrack_FloatCurve.cpp b/Engine/Source/Editor/Persona/Private/AnimTimeline/AnimTimelineTrack_FloatCurve.cpp
--- a/Engine/Source/Editor/Persona/Private/AnimTimeline/AnimTimelineTrack_FloatCurve.cpp
+++ b/Engine/Source/Editor/Persona/Private/AnimTimeline/AnimTimelineTrack_FloatCurve.cpp
@@ -20,7 +20,6 @@
 #include "Widgets/Layout/SBox.h"
 #include "AnimModel_AnimSequenceBase.h"
 #include "SAnimOutlinerItem.h"
-#include "Animation/AnimData/AnimDataController.h"
 
 #define LOCTEXT_NAMESPACE "FAnimTimelineTrack_FloatCurve"
 
@@ -156,10 +155,10 @@ void FAnimTimelineTrack_FloatCurve::ConvertCurveToMetaData()
 	IAnimationEditor::FCurveEditInfo EditInfo(CurveName, ERawCurveTrackTypes::RCT_Float, 0);
 	StaticCastSharedRef<FAnimModel_AnimSequenceBase>(GetModel())->OnStopEditingCurves.ExecuteIfBound(TArray<IAnimationEditor::FCurveEditInfo>({ EditInfo }));
 
-	UAnimDataController* Controller = AnimSequenceBase->GetController();
-	UAnimDataController::FScopedBracket ScopedBracket(Controller, LOCTEXT("ConvertCurveToMetaData_Bracket", "Converting curve to metadata"));
-	Controller->SetCurveFlag(CurveId, AACF_Metadata, true);
-	Controller->SetCurveKeys(CurveId, { FRichCurveKey(0.f, 1.f) });	
+	IAnimationDataController& Controller = AnimSequenceBase->GetController();
+	IAnimationDataController::FScopedBracket ScopedBracket(Controller, LOCTEXT("ConvertCurveToMetaData_Bracket", "Converting curve to metadata"));
+	Controller.SetCurveFlag(CurveId, AACF_Metadata, true);
+	Controller.SetCurveKeys(CurveId, { FRichCurveKey(0.f, 1.f) });	
 
 	ZoomToFit();
 }
@@ -168,9 +167,9 @@ void FAnimTimelineTrack_FloatCurve::ConvertMetaDataToCurve()
 {
 	UAnimSequenceBase* AnimSequenceBase = GetModel()->GetAnimSequenceBase();
 
-	UAnimDataController* Controller = AnimSequenceBase->GetController();
-	UAnimDataController::FScopedBracket ScopedBracket(Controller, LOCTEXT("CurvePanel_ConvertMetaDataToCurve", "Convert metadata to curve"));
-	Controller->SetCurveFlag(CurveId, AACF_Metadata, false);	
+	IAnimationDataController& Controller = AnimSequenceBase->GetController();
+	IAnimationDataController::FScopedBracket ScopedBracket(Controller, LOCTEXT("CurvePanel_ConvertMetaDataToCurve", "Convert metadata to curve"));
+	Controller.SetCurveFlag(CurveId, AACF_Metadata, false);	
 }
 
 void FAnimTimelineTrack_FloatCurve::RemoveCurve()
@@ -188,8 +187,8 @@ void FAnimTimelineTrack_FloatCurve::RemoveCurve()
 
 			AnimSequenceBase->Modify(true);
 
-			UAnimDataController* Controller = AnimSequenceBase->GetController();
-			Controller->RemoveCurve(CurveId);
+			IAnimationDataController& Controller = AnimSequenceBase->GetController();
+			Controller.RemoveCurve(CurveId);
 
 			AnimSequenceBase->PostEditChange();
 		}
@@ -258,12 +257,12 @@ void FAnimTimelineTrack_FloatCurve::OnCommitCurveName(const FText& InText, EText
 				}
 			}
 
-			UAnimDataController* Controller = AnimSequenceBase->GetController();
-			UAnimDataController::FScopedBracket ScopedBracket(Controller, LOCTEXT("CurveEditor_RenameCurve", "Renaming Curve"));
+			IAnimationDataController& Controller = AnimSequenceBase->GetController();
+			IAnimationDataController::FScopedBracket ScopedBracket(Controller, LOCTEXT("CurveEditor_RenameCurve", "Renaming Curve"));
           
 			FAnimationCurveIdentifier NewCurveId(NewSmartName, ERawCurveTrackTypes::RCT_Float);
-			Controller->RenameCurve(CurveId, NewCurveId);
-			Controller->RemoveBoneTracksMissingFromSkeleton(AnimSequenceBase->GetSkeleton());           
+			Controller.RenameCurve(CurveId, NewCurveId);
+			Controller.RemoveBoneTracksMissingFromSkeleton(AnimSequenceBase->GetSkeleton());           
 		}
 	}
 }
@@ -320,7 +319,7 @@ void FAnimTimelineTrack_FloatCurve::AddCurveTrackButton(TSharedPtr<SHorizontalBo
 	auto SetValue = [this](FLinearColor InNewColor)
 	{
 		UAnimSequenceBase* AnimSequenceBase = GetModel()->GetAnimSequenceBase();
-		AnimSequenceBase->GetController()->SetCurveColor(CurveId, InNewColor);
+		AnimSequenceBase->GetController().SetCurveColor(CurveId, InNewColor);
 
 		// Set display curves too
 		for(const TPair<FCurveModelID, TUniquePtr<FCurveModel>>& CurvePair : CurveEditor->GetCurves())
@@ -340,12 +339,12 @@ void FAnimTimelineTrack_FloatCurve::AddCurveTrackButton(TSharedPtr<SHorizontalBo
 			.OnInteractivePickBegin_Lambda([this]()
 			{
 				UAnimSequenceBase* AnimSequenceBase = GetModel()->GetAnimSequenceBase();
-				AnimSequenceBase->GetController()->OpenBracket(LOCTEXT("EditCurveColor_Bracket", "Editing Curve Color"));				
+				AnimSequenceBase->GetController().OpenBracket(LOCTEXT("EditCurveColor_Bracket", "Editing Curve Color"));				
 			})
 			.OnInteractivePickEnd_Lambda([this]()
 			{
 				UAnimSequenceBase* AnimSequenceBase = GetModel()->GetAnimSequenceBase();
-				AnimSequenceBase->GetController()->CloseBracket();
+				AnimSequenceBase->GetController().CloseBracket();
 			});
 	};
 
diff --git a/Engine/Source/Editor/Persona/Private/AnimTimeline/AnimTimelineTrack_TransformCurve.cpp b/Engine/Source/Editor/Persona/Private/AnimTimeline/AnimTimelineTrack_TransformCurve.cpp
--- a/Engine/Source/Editor/Persona/Private/AnimTimeline/AnimTimelineTrack_TransformCurve.cpp
+++ b/Engine/Source/Editor/Persona/Private/AnimTimeline/AnimTimelineTrack_TransformCurve.cpp
@@ -12,7 +12,6 @@
 #include "Animation/DebugSkelMeshComponent.h"
 #include "AnimPreviewInstance.h"
 #include "AnimModel_AnimSequenceBase.h"
-#include "Animation/AnimData/AnimDataController.h"
 
 #define LOCTEXT_NAMESPACE "FAnimTimelineTrack_TransformCurve"
 
@@ -143,8 +142,8 @@ void FAnimTimelineTrack_TransformCurve::DeleteTrack()
 			BaseModel->OnStopEditingCurves.ExecuteIfBound(CurveEditInfo);
 
 
-			UAnimDataController* Controller = AnimSequenceBase->GetController();
-			Controller->RemoveCurve(CurveId);
+			IAnimationDataController& Controller = AnimSequenceBase->GetController();
+			Controller.RemoveCurve(CurveId);
 
 			if (GetModel()->GetPreviewScene()->GetPreviewMeshComponent()->PreviewInstance != nullptr)
 			{
@@ -165,8 +164,8 @@ void FAnimTimelineTrack_TransformCurve::ToggleEnabled()
 {
 	UAnimSequenceBase* AnimSequenceBase = GetModel()->GetAnimSequenceBase();
 
-	UAnimDataController* Controller = AnimSequenceBase->GetController();
-	Controller->SetCurveFlag(CurveId, AACF_Disabled, !IsEnabled());
+	IAnimationDataController& Controller = AnimSequenceBase->GetController();
+	Controller.SetCurveFlag(CurveId, AACF_Disabled, !IsEnabled());
 
 	// need to update curves, otherwise they're not disabled
 	if (GetModel()->GetPreviewScene()->GetPreviewMeshComponent()->PreviewInstance != nullptr)
diff --git a/Engine/Source/Editor/Persona/Private/AnimTimeline/SAnimTimeline.cpp b/Engine/Source/Editor/Persona/Private/AnimTimeline/SAnimTimeline.cpp
--- a/Engine/Source/Editor/Persona/Private/AnimTimeline/SAnimTimeline.cpp
+++ b/Engine/Source/Editor/Persona/Private/AnimTimeline/SAnimTimeline.cpp
@@ -33,7 +33,6 @@
 #include "Widgets/Input/SSpinBox.h"
 #include "SAnimTimelineTransportControls.h"
 #include "Animation/AnimSequence.h"
-#include "Animation/AnimData/AnimDataController.h"
 #include "Animation/AnimData/AnimDataModel.h"
 #include "Runtime/Engine/Classes/Animation/AnimSequenceHelpers.h"
 
@@ -600,8 +599,8 @@ void SAnimTimeline::OnReZeroAnimSequence(int32 FrameIndex)
 					RawTrack.PosKeys[i] += ApplyTranslation;
 				}
 
-				UAnimDataController* Controller = AnimSequence->GetController();
-				Controller->SetBoneTrackKeys(AnimationTrack.Name, RawTrack.PosKeys, RawTrack.RotKeys, RawTrack.ScaleKeys);
+				IAnimationDataController& Controller = AnimSequence->GetController();
+				Controller.SetBoneTrackKeys(AnimationTrack.Name, RawTrack.PosKeys, RawTrack.RotKeys, RawTrack.ScaleKeys);
 
 				AnimSequence->MarkPackageDirty();
 			}
diff --git a/Engine/Source/Editor/Persona/Private/SAnimSequenceCurveEditor.cpp b/Engine/Source/Editor/Persona/Private/SAnimSequenceCurveEditor.cpp
--- a/Engine/Source/Editor/Persona/Private/SAnimSequenceCurveEditor.cpp
+++ b/Engine/Source/Editor/Persona/Private/SAnimSequenceCurveEditor.cpp
@@ -17,7 +17,6 @@
 #include "Editor.h"
 #include "Framework/MultiBox/MultiBoxBuilder.h"
 #include "EditorStyleSet.h"
-#include "Animation/AnimData/AnimDataController.h"
 
 #define LOCTEXT_NAMESPACE "SAnimSequenceCurveEditor"
 
@@ -65,7 +64,7 @@ const FRichCurve& FRichCurveEditorModelNamed::GetReadOnlyRichCurve() const
 
 void FRichCurveEditorModelNamed::CurveHasChanged()
 {
-	UAnimDataController* Controller = AnimSequence->GetController();
+	IAnimationDataController& Controller = AnimSequence->GetController();
 
 	switch (Type)
 	{
@@ -77,7 +76,7 @@ void FRichCurveEditorModelNamed::CurveHasChanged()
 		case ERawCurveTrackTypes::RCT_Transform:
 		case ERawCurveTrackTypes::RCT_Float:
 		{
-			Controller->SetCurveKeys(CurveId, CachedCurve.GetConstRefOfKeys());
+			Controller.SetCurveKeys(CurveId, CachedCurve.GetConstRefOfKeys());
 			break;
 		}
 	}
diff --git a/Engine/Source/Editor/Persona/Private/SAnimationScrubPanel.cpp b/Engine/Source/Editor/Persona/Private/SAnimationScrubPanel.cpp
--- a/Engine/Source/Editor/Persona/Private/SAnimationScrubPanel.cpp
+++ b/Engine/Source/Editor/Persona/Private/SAnimationScrubPanel.cpp
@@ -11,7 +11,6 @@
 #include "Animation/BlendSpace.h"
 #include "AnimationEditorPreviewScene.h"
 #include "Animation/AnimData/AnimDataModel.h"
-#include "Animation/AnimData/AnimDataController.h"
 #include "Animation/AnimSequenceHelpers.h"
 
 #define LOCTEXT_NAMESPACE "AnimationScrubPanel"
@@ -590,8 +589,8 @@ void SAnimationScrubPanel::OnReZeroAnimSequence(int32 FrameIndex)
 					RawTrack.PosKeys[i] += ApplyTranslation;
 				}
 
-				UAnimDataController* Controller = AnimSequence->GetController();
-				Controller->SetBoneTrackKeys(AnimationTrack.Name, RawTrack.PosKeys, RawTrack.RotKeys, RawTrack.ScaleKeys);
+				IAnimationDataController& Controller = AnimSequence->GetController();
+				Controller.SetBoneTrackKeys(AnimationTrack.Name, RawTrack.PosKeys, RawTrack.RotKeys, RawTrack.ScaleKeys);
 
 				AnimSequence->MarkPackageDirty();
 			}
diff --git a/Engine/Source/Editor/SequenceRecorder/Private/AnimationRecorder.cpp b/Engine/Source/Editor/SequenceRecorder/Private/AnimationRecorder.cpp
--- a/Engine/Source/Editor/SequenceRecorder/Private/AnimationRecorder.cpp
+++ b/Engine/Source/Editor/SequenceRecorder/Private/AnimationRecorder.cpp
@@ -19,7 +19,6 @@
 #include "Animation/AnimNotifies/AnimNotify.h"
 #include "Animation/AnimNotifies/AnimNotifyState.h"
 #include "Subsystems/AssetEditorSubsystem.h"
-#include "Animation/AnimData/AnimDataController.h"
 #include "Animation/BuiltInAttributeTypes.h"
 
 #define LOCTEXT_NAMESPACE "FAnimationRecorder"
@@ -233,13 +232,13 @@ void FAnimationRecorder::StartRecord(USkeletalMeshComponent* Component, UAnimSeq
 
 	LastFrame = 0;
 
-	UAnimDataController* Controller = AnimationObject->GetController();
-	Controller->SetModel(AnimationObject->GetDataModel());
+	IAnimationDataController& Controller = AnimationObject->GetController();
+	Controller.SetModel(AnimationObject->GetDataModel());
 
-	Controller->OpenBracket(LOCTEXT("StartRecord_Bracket", "Starting Animation Recording"));
+	Controller.OpenBracket(LOCTEXT("StartRecord_Bracket", "Starting Animation Recording"));
 
 	AnimationObject->DeleteNotifyTrackData();
-	Controller->ResetModel();
+	Controller.ResetModel();
 
 	RecordedCurves.Reset();
 	RecordedTimes.Empty();
@@ -255,7 +254,7 @@ void FAnimationRecorder::StartRecord(USkeletalMeshComponent* Component, UAnimSeq
 		{
 			// add tracks for the bone existing
 			const FName BoneTreeName = AnimSkeleton->GetReferenceSkeleton().GetBoneName(BoneTreeIndex);
-			Controller->AddBoneTrack(BoneTreeName);			
+			Controller.AddBoneTrack(BoneTreeName);			
 			RawTracks.AddDefaulted();
 		}
 	}
@@ -306,15 +305,15 @@ UAnimSequence* FAnimationRecorder::StopRecord(bool bShowMessage)
 
 	if (AnimationObject)
 	{
-		UAnimDataController* Controller = AnimationObject->GetController();
+		IAnimationDataController& Controller = AnimationObject->GetController();
 		int32 NumKeys = LastFrame  + 1;
 
 		// can't use TimePassed. That is just total time that has been passed, not necessarily match with frame count
-		Controller->SetPlayLength( (NumKeys>1) ? (NumKeys-1) * IntervalTime : MINIMUM_ANIMATION_LENGTH );
+		Controller.SetPlayLength( (NumKeys>1) ? (NumKeys-1) * IntervalTime : MINIMUM_ANIMATION_LENGTH );
 
 		const float FloatDenominator = 1000.0f;
 		const float Numerator = FloatDenominator / IntervalTime;
-		Controller->SetFrameRate(FFrameRate(Numerator, FloatDenominator));
+		Controller.SetFrameRate(FFrameRate(Numerator, FloatDenominator));
 
 		ProcessNotifies();
 
@@ -357,7 +356,7 @@ UAnimSequence* FAnimationRecorder::StopRecord(bool bShowMessage)
 								{
 									// give default curve flag for recording 
 									const FAnimationCurveIdentifier CurveId(CurveName, ERawCurveTrackTypes::RCT_Float);
-									Controller->AddCurve(CurveId, AACF_DefaultCurve);
+									Controller.AddCurve(CurveId, AACF_DefaultCurve);
 									FloatCurveData = AnimationObject->GetDataModel()->FindFloatCurve(CurveId);
 
 								}
@@ -386,7 +385,7 @@ UAnimSequence* FAnimationRecorder::StopRecord(bool bShowMessage)
 						}
 
 						const FAnimationCurveIdentifier CurveId(FloatCurveData->Name, ERawCurveTrackTypes::RCT_Float);
-						Controller->SetCurveKeys(CurveId, Keys);
+						Controller.SetCurveKeys(CurveId, Keys);
 					}
 				}
 			}	
@@ -402,11 +401,11 @@ UAnimSequence* FAnimationRecorder::StopRecord(bool bShowMessage)
 			const FBoneAnimationTrack& AnimationTrack = BoneAnimationTracks[TrackIndex];
 			FRawAnimSequenceTrack& RawTrack = RawTracks[TrackIndex];
 
-			Controller->SetBoneTrackKeys(AnimationTrack.Name, RawTrack.PosKeys, RawTrack.RotKeys, RawTrack.ScaleKeys);
+			Controller.SetBoneTrackKeys(AnimationTrack.Name, RawTrack.PosKeys, RawTrack.RotKeys, RawTrack.ScaleKeys);
 		}
 
-		Controller->NotifyPopulated();
-		Controller->CloseBracket();
+		Controller.NotifyPopulated();
+		Controller.CloseBracket();
 
 		AnimationObject->MarkPackageDirty();
 		
@@ -680,7 +679,7 @@ bool FAnimationRecorder::Record(USkeletalMeshComponent* Component, FTransform co
 {
 	if (ensure(AnimationObject))
 	{
-		UAnimDataController* Controller = AnimationObject->GetController();
+		IAnimationDataController& Controller = AnimationObject->GetController();
 		USkeletalMesh* SkeletalMesh = Component->MasterPoseComponent != nullptr ? ToRawPtr(Component->MasterPoseComponent->SkeletalMesh) : ToRawPtr(Component->SkeletalMesh);
 
 		const TArray<FBoneAnimationTrack>& BoneAnimationTracks = AnimationObject->GetDataModel()->GetBoneAnimationTracks();
diff --git a/Engine/Source/Editor/SkeletonEditor/Private/EditableSkeleton.cpp b/Engine/Source/Editor/SkeletonEditor/Private/EditableSkeleton.cpp
--- a/Engine/Source/Editor/SkeletonEditor/Private/EditableSkeleton.cpp
+++ b/Engine/Source/Editor/SkeletonEditor/Private/EditableSkeleton.cpp
@@ -30,7 +30,6 @@
 #include "Engine/DataAsset.h"
 #include "Animation/PreviewCollectionInterface.h"
 #include "HAL/PlatformApplicationMisc.h"
-#include "Animation/AnimData/AnimDataController.h"
 
 const FString FEditableSkeleton::SocketCopyPasteHeader = TEXT("SocketCopyPasteBuffer");
 
@@ -321,10 +320,10 @@ void FEditableSkeleton::RenameSmartname(const FName InContainerName, SmartName::
 
 					if (SequenceBase->GetDataModel()->FindFloatCurve(CurveId))
 					{
-						UAnimDataController* Controller = SequenceBase->GetController();
+						IAnimationDataController& Controller = SequenceBase->GetController();
 
 						FAnimationCurveIdentifier NewCurveId(FSmartName(InNewName, InNameUid), ERawCurveTrackTypes::RCT_Float);
-						Controller->RenameCurve(CurveId, NewCurveId);
+						Controller.RenameCurve(CurveId, NewCurveId);
 						if (UAnimSequence* Seq = Cast<UAnimSequence>(SequenceBase))
 						{
 							SequencesToRecompress.Add(Seq);
@@ -481,15 +480,15 @@ void FEditableSkeleton::RemoveSmartnamesAndFixupAnimations(const FName& InContai
 					USkeleton* MySkeleton = Sequence->GetSkeleton();
 					Sequence->Modify(true);
 
-					UAnimDataController* Controller = Sequence->GetController();
-					UAnimDataController::FScopedBracket ScopedBracket(Controller, LOCTEXT("DeleteCurvesNotOnSkeletonBracket", "Deleting curves that no longer exist on Skeleton"));
+					IAnimationDataController& Controller = Sequence->GetController();
+					IAnimationDataController::FScopedBracket ScopedBracket(Controller, LOCTEXT("DeleteCurvesNotOnSkeletonBracket", "Deleting curves that no longer exist on Skeleton"));
 					for (FName Name : InNames)
 					{
 						FSmartName CurveToDelete;
 						if (MySkeleton->GetSmartNameByName(USkeleton::AnimCurveMappingName, Name, CurveToDelete))
 						{
 							const FAnimationCurveIdentifier CurveId(CurveToDelete, ERawCurveTrackTypes::RCT_Float);
-							Controller->RemoveCurve(CurveId);
+							Controller.RemoveCurve(CurveId);
 						}
 					}
 				}
diff --git a/Engine/Source/Editor/UnrealEd/Private/EditorAnimUtils.cpp b/Engine/Source/Editor/UnrealEd/Private/EditorAnimUtils.cpp
--- a/Engine/Source/Editor/UnrealEd/Private/EditorAnimUtils.cpp
+++ b/Engine/Source/Editor/UnrealEd/Private/EditorAnimUtils.cpp
@@ -19,7 +19,6 @@
 #include "ContentBrowserModule.h"
 #include "Subsystems/AssetEditorSubsystem.h"
 #include "Editor.h"
-#include "Animation/AnimData/AnimDataController.h"
 #include "AnimationBlueprintLibrary.h"
 
 #define LOCTEXT_NAMESPACE "EditorAnimUtils"
@@ -257,8 +256,8 @@ namespace EditorAnimUtils
 
 					// clear transform curves since those curves won't work in new skeleton
 					// since we're deleting curves, mark this rebake flag off
-					UAnimDataController* Controller = AnimSequenceToRetarget->GetController();
-					Controller->RemoveAllCurvesOfType(ERawCurveTrackTypes::RCT_Transform);
+					IAnimationDataController& Controller = AnimSequenceToRetarget->GetController();
+					Controller.RemoveAllCurvesOfType(ERawCurveTrackTypes::RCT_Transform);
 
 					// I can't copy transform curves yet because transform curves need retargeting. 
 					//EditorAnimUtils::CopyAnimCurves(OldSkeleton, NewSkeleton, AssetToRetarget, USkeleton::AnimTrackCurveMappingName, FRawCurveTracks::TransformType);
diff --git a/Engine/Source/Editor/UnrealEd/Private/SkeletalMeshEdit.cpp b/Engine/Source/Editor/UnrealEd/Private/SkeletalMeshEdit.cpp
--- a/Engine/Source/Editor/UnrealEd/Private/SkeletalMeshEdit.cpp
+++ b/Engine/Source/Editor/UnrealEd/Private/SkeletalMeshEdit.cpp
@@ -36,7 +36,6 @@
 #include "ComponentReregisterContext.h"
 #include "Components/SkeletalMeshComponent.h"
 #include "FbxAnimUtils.h"
-#include "Animation/AnimData/AnimDataController.h"
 #include "Animation/AnimSequenceHelpers.h"
 #include "Animation/BuiltInAttributeTypes.h"
 
@@ -1312,25 +1311,25 @@ bool UnFbx::FFbxImporter::ImportCurveToAnimSequence(class UAnimSequence * Target
 
 
 		UAnimDataModel* DataModel = TargetSequence->GetDataModel();
-		UAnimDataController* Controller = TargetSequence->GetController();
+		IAnimationDataController& Controller = TargetSequence->GetController();
 
 		const FFloatCurve* TargetCurve = DataModel->FindFloatCurve(FloatCurveId);
 		if (TargetCurve == nullptr)
 		{
 			// Need to add the curve first
-			Controller->AddCurve(FloatCurveId, AACF_DefaultCurve | CurveFlags);
+			Controller.AddCurve(FloatCurveId, AACF_DefaultCurve | CurveFlags);
 			TargetCurve = DataModel->FindFloatCurve(FloatCurveId);
 		}
 		else
 		{
 			// Need to update any of the flags
-			Controller->SetCurveFlags(FloatCurveId, CurveFlags | TargetCurve->GetCurveTypeFlags());
+			Controller.SetCurveFlags(FloatCurveId, CurveFlags | TargetCurve->GetCurveTypeFlags());
 		}
 		
 		// Should be valid at this point
 		ensure(TargetCurve);
 
-		Controller->UpdateCurveNamesFromSkeleton(Skeleton, ERawCurveTrackTypes::RCT_Float);
+		Controller.UpdateCurveNamesFromSkeleton(Skeleton, ERawCurveTrackTypes::RCT_Float);
 
 		FRichCurve RichCurve;
 		if (ImportCurve(FbxCurve, RichCurve, AnimTimeSpan, ValueScale))
@@ -1341,7 +1340,7 @@ bool UnFbx::FFbxImporter::ImportCurveToAnimSequence(class UAnimSequence * Target
 			}
 
 			// Set actual keys on curve within the model
-			Controller->SetCurveKeys(FloatCurveId, RichCurve.GetConstRefOfKeys());
+			Controller.SetCurveKeys(FloatCurveId, RichCurve.GetConstRefOfKeys());
 
 			return true;
 		}
@@ -1493,14 +1492,14 @@ bool UnFbx::FFbxImporter::ImportAnimation(USkeleton* Skeleton, UAnimSequence * D
 	FbxTime SequenceLength = AnimTimeSpan.GetDuration();
 	float PreviousSequenceLength = DestSeq->GetPlayLength();
 
-	UAnimDataController* Controller = DestSeq->GetController();
-	Controller->OpenBracket(LOCTEXT("ImportAnimation_Bracket", "Importing Animation"));
+	IAnimationDataController& Controller = DestSeq->GetController();
+	Controller.OpenBracket(LOCTEXT("ImportAnimation_Bracket", "Importing Animation"));
 
 	//This destroy all previously imported animation raw data
-	Controller->RemoveAllBoneTracks();
+	Controller.RemoveAllBoneTracks();
 
 	// if you have one pose(thus 0.f duration), it still contains animation, so we'll need to consider that as MINIMUM_ANIMATION_LENGTH time length
-	Controller->SetPlayLength(FGenericPlatformMath::Max<float>(SequenceLength.GetSecondDouble(), MINIMUM_ANIMATION_LENGTH));
+	Controller.SetPlayLength(FGenericPlatformMath::Max<float>(SequenceLength.GetSecondDouble(), MINIMUM_ANIMATION_LENGTH));
 	if(PreviousSequenceLength > MINIMUM_ANIMATION_LENGTH && DestSeq->GetDataModel()->GetNumberOfFloatCurves() > 0)
 	{
 		// The sequence already existed when we began the import. We need to scale the key times for all curves to match the new 
@@ -1511,7 +1510,7 @@ bool UnFbx::FFbxImporter::ImportAnimation(USkeleton* Skeleton, UAnimSequence * D
 			for (const FFloatCurve& Curve : DestSeq->GetDataModel()->GetFloatCurves())
 			{
 				const FAnimationCurveIdentifier CurveId(Curve.Name, ERawCurveTrackTypes::RCT_Float);
-				Controller->ScaleCurve(CurveId, 0.f, ScaleFactor);
+				Controller.ScaleCurve(CurveId, 0.f, ScaleFactor);
 			}
 		}
 	}
@@ -1538,7 +1537,7 @@ bool UnFbx::FFbxImporter::ImportAnimation(USkeleton* Skeleton, UAnimSequence * D
 		for (auto CurveName : CurveSmartNamesToRemove)
 		{
 			const FAnimationCurveIdentifier CurveId(CurveName, ERawCurveTrackTypes::RCT_Float);
-			Controller->RemoveCurve(CurveId);
+			Controller.RemoveCurve(CurveId);
 		}
 	}
 
@@ -1610,7 +1609,7 @@ bool UnFbx::FFbxImporter::ImportAnimation(USkeleton* Skeleton, UAnimSequence * D
 	}
 	else if (CurveAttributeKeyCount > 0)
 	{
-		Controller->SetFrameRate(FFrameRate(ResampleRate, 1));
+		Controller.SetFrameRate(FFrameRate(ResampleRate, 1));
 	}
 
 	// Import bone metadata to AnimSequence
@@ -1619,8 +1618,8 @@ bool UnFbx::FFbxImporter::ImportAnimation(USkeleton* Skeleton, UAnimSequence * D
 		ImportNodeCustomProperties(DestSeq, SkeletonNode, true);
 	}
 
-	Controller->NotifyPopulated();
-	Controller->CloseBracket();
+	Controller.NotifyPopulated();
+	Controller.CloseBracket();
 
 	// Reregister skeletal mesh components so they reflect the updated animation
 	for (TObjectIterator<USkeletalMeshComponent> Iter; Iter; ++Iter)
@@ -1847,8 +1846,8 @@ void UnFbx::FFbxImporter::ImportBoneTracks(USkeleton* Skeleton, FAnimCurveImport
 	const FReferenceSkeleton& RefSkeleton = Skeleton->GetReferenceSkeleton();
 
 	UAnimDataModel* DataModel = DestSeq->GetDataModel();
-	UAnimDataController* Controller = DestSeq->GetController();	
-	Controller->SetFrameRate(FFrameRate(ResampleRate, 1));
+	IAnimationDataController& Controller = DestSeq->GetController();	
+	Controller.SetFrameRate(FFrameRate(ResampleRate, 1));
 
 	for (int32 SourceTrackIdx = 0; SourceTrackIdx < AnimImportSettings.FbxRawBoneNames.Num(); ++SourceTrackIdx)
 	{
@@ -1967,8 +1966,8 @@ void UnFbx::FFbxImporter::ImportBoneTracks(USkeleton* Skeleton, FAnimCurveImport
 			if (bSuccess)
 			{
 				//add new track
-				Controller->AddBoneTrack(BoneName);
-				Controller->SetBoneTrackKeys(BoneName, RawTrack.PosKeys, RawTrack.RotKeys, RawTrack.ScaleKeys);
+				Controller.AddBoneTrack(BoneName);
+				Controller.SetBoneTrackKeys(BoneName, RawTrack.PosKeys, RawTrack.RotKeys, RawTrack.ScaleKeys);
 				const int32 TrackIndex = DestSeq->GetDataModel()->GetBoneTrackIndexByName(BoneName);
 
 				NewDebugData.SetTrackData(TrackIndex, BoneTreeIndex, BoneName);
diff --git a/Engine/Source/Editor/UnrealEd/Public/ChangeTransactor.h b/Engine/Source/Editor/UnrealEd/Public/ChangeTransactor.h
--- a/Engine/Source/Editor/UnrealEd/Public/ChangeTransactor.h
+++ b/Engine/Source/Editor/UnrealEd/Public/ChangeTransactor.h
@@ -15,7 +15,7 @@ namespace UE {
 	The operations are targeted a specific UObject, which is used as the 'transactional' object. The transactions themselves 
 	are initiated using the ITransactor API using GEngine as the instance.
 */
-struct ENGINE_API FChangeTransactor
+struct FChangeTransactor
 {
 	FChangeTransactor() : TransactionObject(nullptr), PendingTransactionIndex(INDEX_NONE), TransactedChanges(0) {}
 	FChangeTransactor(UObject* InTransactionObject) : TransactionObject(InTransactionObject), PendingTransactionIndex(INDEX_NONE), TransactedChanges(0) {}
diff --git a/Engine/Source/Runtime/Engine/Classes/Animation/AnimData/AnimDataController.h b/Engine/Source/Runtime/Engine/Classes/Animation/AnimData/AnimDataController.h
deleted file mode 100644
diff --git a/Engine/Source/Runtime/Engine/Classes/Animation/AnimData/IAnimationDataController.h b/Engine/Source/Runtime/Engine/Classes/Animation/AnimData/IAnimationDataController.h
new file mode 100644
--- /dev/null
+++ b/Engine/Source/Runtime/Engine/Classes/Animation/AnimData/IAnimationDataController.h
@@ -0,0 +1,593 @@
+// Copyright Epic Games, Inc. All Rights Reserved.
+
+#pragma once
+
+#include "CoreTypes.h"
+#include "UObject/Interface.h"
+#include "CurveIdentifier.h"
+#include "IAnimationDataController.generated.h"
+
+class UAssetUserData;
+
+class UAnimDataModel;
+
+namespace UE {
+namespace Anim {
+	class FOpenBracketAction;
+	class FCloseBracketAction;
+}}
+
+/**
+ * The Controller is the sole authority to perform changes on the Animation Data Model. Any mutation to the model made will
+ * cause a subsequent notify (EAnimDataModelNotifyType) to be broadcasted from the Model's ModifiedEvent. Alongside of it is a 
+ * payload containing information relevant to the mutation. These notifies should be relied upon to update any dependent views 
+ * or generated (derived) data.
+ */
+UINTERFACE(BlueprintType, meta=(CannotImplementInterfaceInBlueprint))
+class ENGINE_API UAnimationDataController : public UInterface
+{
+	GENERATED_BODY()
+};
+
+class ENGINE_API IAnimationDataController
+{
+public:
+	GENERATED_BODY()
+
+#if WITH_EDITOR
+	/** RAII helper to define a scoped-based bracket, opens and closes a controller bracket automatically */
+    struct FScopedBracket
+	{
+		FScopedBracket(IAnimationDataController* InController, const FText& InDescription)
+            : Controller(*InController)
+		{
+			Controller.OpenBracket(InDescription);
+		}
+
+		FScopedBracket(IAnimationDataController& InController, const FText& InDescription)
+            : Controller(InController)
+		{
+			Controller.OpenBracket(InDescription);
+		}
+		
+		FScopedBracket(TScriptInterface<IAnimationDataController>& InController, const FText& InDescription)
+            : Controller(*InController)
+		{
+			Controller.OpenBracket(InDescription);
+		}
+		
+		~FScopedBracket()
+		{
+			Controller.CloseBracket();
+		}
+	private:
+		IAnimationDataController& Controller;
+	};
+#endif // WITH_EDITOR
+
+	/**
+	* Sets the AnimDataModel instance this controller is supposed to be targeting
+	*
+	* @param	InModel		UAnimDataModel instance to target
+	*/
+	UFUNCTION(BlueprintCallable, Category = AnimationData)
+	virtual void SetModel(UAnimDataModel* InModel) = 0;
+
+	/**
+	* @return		The AnimDataModel instance this controller is currently targeting
+	*/
+	UFUNCTION(BlueprintCallable, Category = AnimationData)
+	virtual UAnimDataModel* GetModel() = 0;
+
+	/**
+	* @return		The AnimDataModel instance this controller is currently targeting
+	*/
+	virtual const UAnimDataModel* const GetModel() const = 0;
+
+	/**
+	* Opens an interaction bracket, used for combining a set of controller actions. Broadcasts a EAnimDataModelNotifyType::BracketOpened notify,
+	* this can be used by any Views or dependendent systems to halt any unnecessary or invalid operations until the (last) bracket is closed.
+	*
+	* @param	InTitle				Description of the bracket, e.g. "Generating Curve Data"
+	* @param	bShouldTransact		Whether or not any undo-redo changes should be generated
+	*/
+	UFUNCTION(BlueprintCallable, Category = AnimationData)
+	virtual void OpenBracket(const FText& InTitle, bool bShouldTransact = true) = 0;
+
+	/**
+	* Closes a previously opened interaction bracket, used for combining a set of controller actions. Broadcasts a EAnimDataModelNotifyType::BracketClosed notify.
+	*
+	* @param	InTitle				Description of the bracket, e.g. "Generating Curve Data"
+	* @param	bShouldTransact		Whether or not any undo-redo changes should be generated
+	*/
+	UFUNCTION(BlueprintCallable, Category = AnimationData)
+	virtual void CloseBracket(bool bShouldTransact = true) = 0;
+
+	/**
+	* Sets the total play-able length in seconds. Broadcasts a EAnimDataModelNotifyType::SequenceLengthChanged notify if successful.
+	* The number of frames and keys for the provided length is recalculated according to the current value of UAnimDataModel::FrameRate.
+	*
+	* @param	Length				New play-able length value, has to be positive and non-zero
+	* @param	bShouldTransact		Whether or not any undo-redo changes should be generated
+	*/
+	UFUNCTION(BlueprintCallable, Category = AnimationData)
+	virtual void SetPlayLength(float Length, bool bShouldTransact = true) = 0;
+
+	/*** Sets the total play-able length in seconds. Broadcasts a EAnimDataModelNotifyType::SequenceLengthChanged notify if successful.
+	* T0 and T1 are expected to represent the window of time that was either added or removed. E.g. for insertion T0 indicates the time
+	* at which additional time starts and T1 were it ends. For removal T0 indicates the time at which time should be started to remove, and T1 indicates the end. Giving a total of T1 - T0 added or removed length.
+	* The number of frames and keys for the provided length is recalculated according to the current value of UAnimDataModel::FrameRate.
+	* @param	Length				Total new play-able length value, has to be positive and non-zero
+	* @param	T0					Point between 0 and Length at which the change in time starts
+	* @param	T1					Point between 0 and Length at which the change in time ends
+	* @param	bShouldTransact		Whether or not any undo-redo changes should be generated
+	*/
+	UFUNCTION(BlueprintCallable, Category = AnimationData)
+	virtual void ResizePlayLength(float NewLength, float T0, float T1, bool bShouldTransact = true) = 0;
+
+	/**
+	* Sets the total play-able length in seconds and resizes curves. Broadcasts EAnimDataModelNotifyType::SequenceLengthChanged
+	* and EAnimDataModelNotifyType::CurveChanged notifies if successful.
+	* T0 and T1 are expected to represent the window of time that was either added or removed. E.g. for insertion T0 indicates the time
+	* at which additional time starts and T1 were it ends. For removal T0 indicates the time at which time should be started to remove, and T1 indicates the end. Giving a total of T1 - T0 added or removed length.
+	* The number of frames and keys for the provided length is recalculated according to the current value of UAnimDataModel::FrameRate.
+	*
+	* @param	Length				Total new play-able length value, has to be positive and non-zero
+	* @param	T0					Point between 0 and Length at which the change in time starts
+	* @param	T1					Point between 0 and Length at which the change in time ends
+	* @param	bShouldTransact		Whether or not any undo-redo changes should be generated
+	*/
+	UFUNCTION(BlueprintCallable, Category = AnimationData)
+	virtual void Resize(float Length, float T0, float T1, bool bShouldTransact = true) = 0;
+
+	/**
+	* Sets the frame rate according to which the bone animation is expected to be sampled. Broadcasts a EAnimDataModelNotifyType::FrameRateChanged notify if successful.
+	* The number of frames and keys for the provided frame rate is recalculated according to the current value of UAnimDataModel::PlayLength.
+	*
+	* @param	FrameRate			The new sampling frame rate, has to be positive and non-zero
+	* @param	bShouldTransact		Whether or not any undo-redo changes should be generated
+	*/
+	UFUNCTION(BlueprintCallable, Category = AnimationData)
+	virtual void SetFrameRate(FFrameRate FrameRate, bool bShouldTransact = true) = 0;
+
+	/**
+	* Adds a new bone animation track for the provided name. Broadcasts a EAnimDataModelNotifyType::TrackAdded notify if successful.
+	*
+	* @param	BoneName			Bone name for which a track should be added
+	* @param	bShouldTransact		Whether or not any undo-redo changes should be generated
+	*
+	* @return	The index at which the bone track was added, INDEX_NONE if adding it failed
+	*/
+	UFUNCTION(BlueprintCallable, Category = AnimationData)
+	virtual int32 AddBoneTrack(FName BoneName, bool bShouldTransact = true) = 0;
+
+	/**
+	* Inserts a new bone animation track for the provided name, at the provided index. Broadcasts a EAnimDataModelNotifyType::TrackAdded notify if successful.
+	* The bone name is verified with the AnimModel's outer target USkeleton to ensure the bone exists.
+	*
+	* @param	BoneName			Bone name for which a track should be inserted
+	* @param	DesiredIndex		Index at which the track should be inserted
+	* @param	bShouldTransact		Whether or not any undo-redo changes should be generated
+	*
+	* @return	The index at which the bone track was inserted, INDEX_NONE if the insertion failed
+	*/
+	UFUNCTION(BlueprintCallable, Category = AnimationData)
+	virtual int32 InsertBoneTrack(FName BoneName, int32 DesiredIndex, bool bShouldTransact = true) = 0;
+
+	/**
+	* Removes an existing bone animation track with the provided name. Broadcasts a EAnimDataModelNotifyType::TrackRemoved notify if successful.
+	*
+	* @param	BoneName			Bone name of the track which should be removed
+	* @param	bShouldTransact		Whether or not any undo-redo changes should be generated
+	*
+	* @return	Whether or not the removal was successful
+	*/
+	UFUNCTION(BlueprintCallable, Category = AnimationData)
+	virtual bool RemoveBoneTrack(FName BoneName, bool bShouldTransact = true) = 0;
+
+	/**
+	* Removes all existing Bone Animation tracks. Broadcasts a EAnimDataModelNotifyType::TrackRemoved for each removed track, wrapped within BracketOpened/BracketClosed notifies.
+	*
+	* @param	bShouldTransact		Whether or not any undo-redo changes should be generated
+	*/
+	UFUNCTION(BlueprintCallable, Category = AnimationData)
+	virtual void RemoveAllBoneTracks(bool bShouldTransact = true) = 0;
+
+	/**
+	* Removes an existing bone animation track with the provided name. Broadcasts a EAnimDataModelNotifyType::TrackChanged notify if successful.
+	* The provided number of keys provided is expected to match for each component, and be non-zero.
+	*
+	* @param	BoneName			Bone name of the track for which the keys should be set
+	* @param	PositionalKeys		Array of keys for the translation component
+	* @param	RotationalKeys		Array of keys for the rotation component
+	* @param	ScalingKeys			Array of keys for the scale component
+	* @param	bShouldTransact		Whether or not any undo-redo changes should be generated
+	*
+	* @return	Whether or not the keys were successfully set
+	*/
+	UFUNCTION(BlueprintCallable, Category = AnimationData)
+	virtual bool SetBoneTrackKeys(FName BoneName, const TArray<FVector>& PositionalKeys, const TArray<FQuat>& RotationalKeys, const TArray<FVector>& ScalingKeys, bool bShouldTransact = true) = 0;
+	
+	/**
+	* Adds a new curve with the provided information. Broadcasts a EAnimDataModelNotifyType::CurveAdded notify if successful.
+	*
+	* @param	CurveId				Identifier for the to-be-added curve
+	* @param	CurveFlags			Flags to be set for the curve
+	* @param	bShouldTransact		Whether or not any undo-redo changes should be generated
+	*
+	* @return	Whether or not the curve was successfully added
+	*/
+	UFUNCTION(BlueprintCallable, Category = CurveData)
+	virtual bool AddCurve(const FAnimationCurveIdentifier& CurveId, int32 CurveFlags = 0x00000004, bool bShouldTransact = true) = 0;
+
+	/**
+	* Duplicated the curve with the identifier. Broadcasts a EAnimDataModelNotifyType::CurveAdded notify if successful.
+	*
+	* @param	CopyCurveId			Identifier for the to-be-duplicated curve
+	* @param	NewCurveId			Identifier for the to-be-added curve
+	* @param	bShouldTransact		Whether or not any undo-redo changes should be generated
+	*
+	* @return	Whether or not the curve was successfully duplicated
+	*/
+	UFUNCTION(BlueprintCallable, Category = CurveData)
+	virtual bool DuplicateCurve(const FAnimationCurveIdentifier& CopyCurveId, const FAnimationCurveIdentifier& NewCurveId, bool bShouldTransact = true) = 0;
+	
+
+	/**
+	* Remove the curve with provided identifier. Broadcasts a EAnimDataModelNotifyType::CurveRemoved notify if successful.
+	*
+	* @param	CopyCurveId			Identifier for the to-be-removed curve
+	* @param	bShouldTransact		Whether or not any undo-redo changes should be generated
+	*
+	* @return	Whether or not the curve was successfully removed
+	*/
+	UFUNCTION(BlueprintCallable, Category = CurveData)
+	virtual bool RemoveCurve(const FAnimationCurveIdentifier& CurveId, bool bShouldTransact = true) = 0;
+
+	/**
+	* Removes all the curves of the provided type. Broadcasts a EAnimDataModelNotifyType::CurveRemoved for each removed curve, wrapped within BracketOpened/BracketClosed notifies.
+	*
+	* @param	SupportedCurveType	Type for which all curves are to be removed
+	* @param	bShouldTransact		Whether or not any undo-redo changes should be generated
+	*/
+	UFUNCTION(BlueprintCallable, Category = CurveData)
+	virtual void RemoveAllCurvesOfType(ERawCurveTrackTypes SupportedCurveType, bool bShouldTransact = true) = 0;
+
+	/**
+	* Set an individual flag for the curve with provided identifier. Broadcasts a EAnimDataModelNotifyType::CurveFlagsChanged notify if successful.
+	*
+	* @param	CurveId			    Identifier for the curve for which the flag state is to be set
+	* @param	Flag				Flag for which the state is supposed to be set
+	* @param	bState				State of the flag to be, true=set/false=not set
+	* @param	bShouldTransact		Whether or not any undo-redo changes should be generated
+	*
+	* @return	Whether or not the flag state was successfully set
+	*/
+	UFUNCTION(BlueprintCallable, Category = CurveData)
+	virtual bool SetCurveFlag(const FAnimationCurveIdentifier& CurveId, EAnimAssetCurveFlags Flag, bool bState = true, bool bShouldTransact = true) = 0;
+
+	/**
+	* Replace the flags for the curve with provided identifier. Broadcasts a EAnimDataModelNotifyType::CurveFlagsChanged notify if successful.
+	*
+	* @param	CurveId			    Identifier for the curve for which the flag state is to be set
+	* @param	Flags				Flag mask with which the existings flags are to be replaced
+	* @param	bShouldTransact		Whether or not any undo-redo changes should be generated
+	*
+	* @return	Whether or not the flag mask was successfully set
+	*/
+	UFUNCTION(BlueprintCallable, Category = CurveData)
+	virtual bool SetCurveFlags(const FAnimationCurveIdentifier& CurveId, int32 Flags, bool bShouldTransact = true) = 0;
+
+	/**
+	* Replace the keys for the transform curve with provided identifier. Broadcasts a EAnimDataModelNotifyType::CurveChanged notify if successful.
+	*
+	* @param	CurveId			    Identifier for the transform curve for which the keys are to be set
+	* @param	TransformValues		Transform Values with which the existing values are to be replaced
+	* @param	TimeKeys			Time Keys with which the existing keys are to be replaced
+	* @param	bShouldTransact		Whether or not any undo-redo changes should be generated
+	*
+	* @return	Whether or not the transform curve keys were successfully set
+	*/
+	UFUNCTION(BlueprintCallable, Category = CurveData)
+	virtual bool SetTransformCurveKeys(const FAnimationCurveIdentifier& CurveId, const TArray<FTransform>& TransformValues, const TArray<float>& TimeKeys, bool bShouldTransact = true) = 0;
+		
+	/**
+	* Sets a single key for the transform curve with provided identifier. Broadcasts a EAnimDataModelNotifyType::CurveChanged notify if successful.
+	* In case a key for any of the individual transform channel curves already exists the value is replaced.
+	*
+	* @param	CurveId			    Identifier for the transform curve for which the key is to be set
+	* @param	Time				Time of the key to be set
+	* @param	Value				Value of the key to be set
+	* @param	bShouldTransact		Whether or not any undo-redo changes should be generated
+	*
+	* @return	Whether or not the transform curve key was successfully set
+	*/
+	UFUNCTION(BlueprintCallable, Category = CurveData)
+	virtual bool SetTransformCurveKey(const FAnimationCurveIdentifier& CurveId, float Time, const FTransform& Value, bool bShouldTransact = true) = 0;
+
+	/**
+	* Removes a single key for the transform curve with provided identifier. Broadcasts a EAnimDataModelNotifyType::CurveChanged notify if successful.
+	*
+	* @param	CurveId			    Identifier for the transform curve for which the key is to be removed
+	* @param	Time				Time of the key to be removed
+	* @param	bShouldTransact		Whether or not any undo-redo changes should be generated
+	*
+	* @return	Whether or not the transform curve key was successfully removed
+	*/
+	UFUNCTION(BlueprintCallable, Category = CurveData)
+	virtual bool RemoveTransformCurveKey(const FAnimationCurveIdentifier& CurveId, float Time, bool bShouldTransact = true) = 0;
+
+	/**
+	* Renames the curve with provided identifier. Broadcasts a EAnimDataModelNotifyType::CurveRenamed notify if successful.
+	*
+	* @param	CurveToRenameId		Identifier for the curve to be renamed
+	* @param	NewCurveId			Time of the key to be removed
+	* @param	bShouldTransact		Whether or not any undo-redo changes should be generated
+	*
+	* @return	Whether or not the curve was successfully renamed
+	*/
+	UFUNCTION(BlueprintCallable, Category = CurveData)
+	virtual bool RenameCurve(const FAnimationCurveIdentifier& CurveToRenameId, const FAnimationCurveIdentifier& NewCurveId, bool bShouldTransact = true) = 0;
+
+	/**
+	* Changes the color of the curve with provided identifier. Broadcasts a EAnimDataModelNotifyType::CurveRenamed notify if successful.
+	* Currently changing curve colors is only supported for float curves.
+	*
+	* @param	CurveId				Identifier of the curve to change the color for
+	* @param	Color				Color to which the curve is to be set
+	* @param	bShouldTransact		Whether or not any undo-redo changes should be generated
+	*
+	* @return	Whether or not the curve color was successfully changed
+	*/
+	UFUNCTION(BlueprintCallable, Category = CurveData)
+	virtual bool SetCurveColor(const FAnimationCurveIdentifier& CurveId, FLinearColor Color, bool bShouldTransact = true) = 0;
+	
+	/**
+	* Scales the curve with provided identifier. Broadcasts a EAnimDataModelNotifyType::CurveScaled notify if successful.
+	*
+	* @param	CurveId				Identifier of the curve to scale
+	* @param	Origin				Time to use as the origin when scaling the curve
+	* @param	Factor				Factor with which the curve is supposed to be scaled
+	* @param	bShouldTransact		Whether or not any undo-redo changes should be generated
+	*
+	* @return	Whether or not scaling the curve was successful
+	*/
+	UFUNCTION(BlueprintCallable, Category = CurveData)
+	virtual bool ScaleCurve(const FAnimationCurveIdentifier& CurveId, float Origin, float Factor, bool bShouldTransact = true) = 0;
+
+	/**
+	* Sets a single key for the curve with provided identifier and name. Broadcasts a EAnimDataModelNotifyType::CurveChanged notify if successful.
+	* In case a key for the provided key time already exists the key is replaced.
+	*
+	* @param	CurveId			    Identifier for the curve for which the key is to be set
+	* @param	Key					Key to be set
+	* @param	bShouldTransact		Whether or not any undo-redo changes should be generated
+	*
+	* @return	Whether or not the curve key was successfully set
+	*/
+	UFUNCTION(BlueprintCallable, Category = CurveData)
+	virtual bool SetCurveKey(const FAnimationCurveIdentifier& CurveId, const FRichCurveKey& Key, bool bShouldTransact = true) = 0;
+	
+	/**
+	* Remove a single key from the curve with provided identifier and name. Broadcasts a EAnimDataModelNotifyType::CurveChanged notify if successful.
+	*
+	* @param	CurveId			    Identifier for the curve for which the key is to be removed
+	* @param	Time				Time of the key to be removed
+	* @param	bShouldTransact		Whether or not any undo-redo changes should be generated
+	*
+	* @return	Whether or not the curve key was successfully removed
+	*/
+	UFUNCTION(BlueprintCallable, Category = CurveData)
+	virtual bool RemoveCurveKey(const FAnimationCurveIdentifier& CurveId, float Time, bool bShouldTransact = true) = 0;
+
+	/**
+	* Replace the keys for the curve with provided identifier and name. Broadcasts a EAnimDataModelNotifyType::CurveChanged notify if successful.
+	*
+	* @param	CurveId			    Identifier for the curve for which the keys are to be replaced
+	* @param	CurveKeys			Keys with which the existing keys are to be replaced
+	* @param	bShouldTransact		Whether or not any undo-redo changes should be generated
+	*
+	* @return	Whether or not replacing curve keys was succesful
+	*/
+	UFUNCTION(BlueprintCallable, Category = CurveData)
+	virtual bool SetCurveKeys(const FAnimationCurveIdentifier& CurveId, const TArray<FRichCurveKey>& CurveKeys, bool bShouldTransact = true) = 0;
+
+	/**
+	* Updates the display name values for any stored curve, with the names being retrieved from the provided skeleton. Broadcasts a EAnimDataModelNotifyType::CurveRenamed for each to-be-updated curve name, wrapped within BracketOpened/BracketClosed notifies.
+	*
+	* @param	Skeleton			Skeleton to retrieve the display name values from
+	* @param	SupportedCurveType	Curve type for which the names should be updated
+	* @param	bShouldTransact		Whether or not any undo-redo changes should be generated
+	*/
+	UFUNCTION(BlueprintCallable, Category = CurveData)
+	virtual void UpdateCurveNamesFromSkeleton(const USkeleton* Skeleton, ERawCurveTrackTypes SupportedCurveType, bool bShouldTransact = true) = 0;
+
+	/**
+	* Updates the curve names with the provided skeleton, if a display name is not found it will be added thus modifying the skeleton. Broadcasts a EAnimDataModelNotifyType::CurveRenamed for each curve name for which the UID was different or if it was added as a new smartname, wrapped within BracketOpened/BracketClosed notifies.
+	*
+	* @param	Skeleton			Skeleton to retrieve the display name values from
+	* @param	SupportedCurveType	Curve type for which the names should be updated
+	* @param	bShouldTransact		Whether or not any undo-redo changes should be generated
+	*/
+	UFUNCTION(BlueprintCallable, Category = CurveData)
+	virtual void FindOrAddCurveNamesOnSkeleton(USkeleton* Skeleton, ERawCurveTrackTypes SupportedCurveType, bool bShouldTransact = true) = 0;
+
+	/**
+	* Removes any bone track for which the name was not found in the provided skeleton. Broadcasts a EAnimDataModelNotifyType::TrackRemoved for each track which was not found in the skeleton, wrapped within BracketOpened/BracketClosed notifies.
+	*
+	* @param	Skeleton			Skeleton to retrieve the display name values from
+	* @param	bShouldTransact		Whether or not any undo-redo changes should be generated
+	*/
+	virtual bool RemoveBoneTracksMissingFromSkeleton(const USkeleton* Skeleton, bool bShouldTransact = true) = 0;
+
+	/**
+	* Removes any bone attribute for which the name was not found in the provided skeleton. Broadcasts a EAnimDataModelNotifyType::AttributeRemoved for each attribute which was not found in the skeleton, wrapped within BracketOpened/BracketClosed notifies.
+	* Updates any bone attribute for which the bone index is different in the provided skeleton. Broadcasts a EAnimDataModelNotifyType::AttributeAdded and EAnimDataModelNotifyType::AttributeRemove for each attribute which was remapped
+	*
+	* @param	Skeleton			Skeleton to retrieve the bone information from
+	* @param	bShouldTransact		Whether or not any undo-redo changes should be generated
+	*/
+	void UpdateAttributesFromSkeleton(const USkeleton* Skeleton, bool bShouldTransact = true);
+
+	/**
+	* Broadcast a EAnimDataModelNotifyType::Populated notify.
+	*/
+	virtual void NotifyPopulated() = 0;	
+
+	/**
+	* Resets all data stored in the model, broadcasts a EAnimDataModelNotifyType::Reset and wraps all actions within BracketOpened/BracketClosed notifies.
+	*	- Bone tracks, broadcasts a EAnimDataModelNotifyType::TrackRemoved for each
+	*	- Curves, broadcasts a EAnimDataModelNotifyType::CurveRemoves for each
+	*	- Play length to one frame at 30fps, broadcasts a EAnimDataModelNotifyType::PlayLengthChanged
+	*	- Frame rate to 30fps, broadcasts a EAnimDataModelNotifyType::FrameRateChanged
+	*
+	* @param	bShouldTransact		Whether or not any undo-redo changes should be generated
+	*/	
+	virtual void ResetModel(bool bShouldTransact = true) = 0;
+
+	/**
+	* Adds a new attribute with the provided information. Broadcasts a EAnimDataModelNotifyType::AttributeAdded notify if successful.
+	*
+	* @param	AttributeIdentifier		Identifier for the to-be-added attribute
+	* @param	bShouldTransact			Whether or not any undo-redo changes should be generated
+	*
+	* @return	Whether or not the attribute was successfully added
+	*/
+	UFUNCTION(BlueprintCallable, Category = CurveData)
+	bool AddAttribute(const FAnimationAttributeIdentifier& AttributeIdentifier, bool bShouldTransact = true);
+
+	/**
+	* Removes an attribute, if found, with the provided information. Broadcasts a EAnimDataModelNotifyType::AttributeRemoved notify if successful.
+	*
+	* @param	AttributeIdentifier		Identifier for the to-be-removed attribute
+	* @param	bShouldTransact			Whether or not any undo-redo changes should be generated
+	*
+	* @return	Whether or not the attribute was successfully removed
+	*/
+	UFUNCTION(BlueprintCallable, Category = CurveData)
+	bool RemoveAttribute(const FAnimationAttributeIdentifier& AttributeIdentifier, bool bShouldTransact = true);
+
+	/**
+	* Removes all attributes for the specified bone name, if any. Broadcasts a EAnimDataModelNotifyType::AttributeRemoved notify for each removed attribute.
+	*
+	* @param	BoneName			Name of the bone to remove attributes for
+	* @param	bShouldTransact		Whether or not any undo-redo changes should be generated
+	*
+	* @return	Total number of removes attributes
+	*/
+	UFUNCTION(BlueprintCallable, Category = CurveData)
+	int32 RemoveAllAttributesForBone(const FName& BoneName, bool bShouldTransact = true);
+
+	/**
+	* Removes all stored attributes. Broadcasts a EAnimDataModelNotifyType::AttributeRemoved notify for each removed attribute.
+	*
+	* @param	bShouldTransact		Whether or not any undo-redo changes should be generated
+	*
+	* @return	Total number of removes attributes
+	*/
+	UFUNCTION(BlueprintCallable, Category = CurveData)
+	int32 RemoveAllAttributes(bool bShouldTransact = true);	
+
+	/**
+	* Sets a single key for the attribute with provided identifier. Broadcasts a EAnimDataModelNotifyType::AttributeChanged notify if successful.
+	* In case a key for the provided key time already exists the key is replaced.
+	*
+	* @param	AttributeIdentifier		Identifier for the attribute for which the key is to be set
+	* @param	Time					Time of the to-be-set key
+	* @param	KeyValue				Value (templated) of the to-be-set key
+	* @param	bShouldTransact			Whether or not any undo-redo changes should be generated
+	*
+	* @return	Whether or not the key was successfully set
+	*/
+	template<typename AttributeType>
+	bool SetTypedAttributeKey(const FAnimationAttributeIdentifier& AttributeIdentifier, float Time, AttributeType& KeyValue, bool bShouldTransact = true)
+	{
+		return SetAttributeKey_Internal(AttributeIdentifier, Time, (const void*)&KeyValue, AttributeType::StaticStruct(), bShouldTransact);
+	}
+
+	/**
+	* Sets a single key for the attribute with provided identifier. Broadcasts a EAnimDataModelNotifyType::AttributeChanged notify if successful.
+	* In case a key for the provided key time already exists the key is replaced.
+	*
+	* @param	AttributeIdentifier		Identifier for the attribute for which the key is to be set
+	* @param	Time					Time of the to-be-set key
+	* @param	KeyValue				Value of the to-be-set key
+	* @param	bShouldTransact			Whether or not any undo-redo changes should be generated
+	*
+	* @return	Whether or not the key was successfully set
+	*/
+	bool SetAttributeKey(const FAnimationAttributeIdentifier& AttributeIdentifier, float Time, const void* KeyValue, bool bShouldTransact = true)
+	{
+		return SetAttributeKey_Internal(AttributeIdentifier, Time, KeyValue, AttributeIdentifier.GetType(), bShouldTransact);
+	}
+	
+	/**
+	* Replace the keys for the attribute with provided identifier. Broadcasts a EAnimDataModelNotifyType::AttributeChanged notify if successful.
+	*
+	* @param	AttributeIdentifier		Identifier for the attribute for which the keys are to be replaced
+	* @param	Times					Times with which the existing key timings are to be replaced
+	* @param	KeyValues				Values with which the existing key values are to be replaced
+	* @param	bShouldTransact			Whether or not any undo-redo changes should be generated
+	*
+	* @return	Whether or not replacing the attribute keys was successful
+	*/
+	bool SetAttributeKeys(const FAnimationAttributeIdentifier& AttributeIdentifier, TArrayView<const float> Times, TArrayView<const void*> KeyValues, bool bShouldTransact = true)
+	{
+		return SetAttributeKeys_Internal(AttributeIdentifier, Times, KeyValues, AttributeIdentifier.GetType(), bShouldTransact);
+	}
+
+	/**
+	* Replace the keys for the attribute with provided identifier. Broadcasts a EAnimDataModelNotifyType::AttributeChanged notify if successful.
+	*
+	* @param	AttributeIdentifier		Identifier for the attribute for which the keys are to be replaced
+	* @param	Times					Times with which the existing key timings are to be replaced
+	* @param	KeyValues				Values (templated) with which the existing key values are to be replaced
+	* @param	bShouldTransact			Whether or not any undo-redo changes should be generated
+	*
+	* @return	Whether or not replacing the attribute keys was successful
+	*/
+	template<typename AttributeType>
+	bool SetTypedAttributeKeys(const FAnimationAttributeIdentifier& AttributeIdentifier, TArrayView<const float> Times, TArrayView<AttributeType> KeyValues, bool bShouldTransact = true)
+	{		
+		TArray<const void*> KeyValuePtrs;
+		Algo::Transform(KeyValues, KeyValuePtrs, [](const AttributeType& Value)
+		{
+			return (const void*)&Value;
+		});
+
+		return SetAttributeKeys_Internal(AttributeIdentifier, Times, MakeArrayView(KeyValuePtrs), AttributeType::StaticStruct(), bShouldTransact);
+	}
+
+	/**
+	* Remove a single key from the attribute with provided identifier. Broadcasts a EAnimDataModelNotifyType::AttributeChanged notify if successful.
+	*
+	* @param	AttributeIdentifier		Identifier for the attribute from which the key is to be removed
+	* @param	Time					Time of the key to be removed
+	* @param	bShouldTransact			Whether or not any undo-redo changes should be generated
+	*
+	* @return	Whether or not the attribute key was successfully removed
+	*/
+	UFUNCTION(BlueprintCallable, Category = CurveData)
+	bool RemoveAttributeKey(const FAnimationAttributeIdentifier& AttributeIdentifier, float Time, bool bShouldTransact = true);
+
+	/**
+	* Duplicated the attribute (curve) with the identifier. Broadcasts a EAnimDataModelNotifyType::AttributeAdded notify if successful.
+	*
+	* @param	AttributeIdentifier			Identifier for the to-be-duplicated attribute
+	* @param	NewAttributeIdentifier		Identifier for the to-be-added attribute
+	* @param	bShouldTransact				Whether or not any undo-redo changes should be generated
+	*
+	* @return	Whether or not the attribute was successfully duplicated
+	*/
+	UFUNCTION(BlueprintCallable, Category = CurveData)
+	bool DuplicateAttribute(const FAnimationAttributeIdentifier& AttributeIdentifier, const FAnimationAttributeIdentifier& NewAttributeIdentifier, bool bShouldTransact = true);
+protected:
+	/** Functionality used by FOpenBracketAction and FCloseBracketAction to broadcast their equivalent notifies without actually opening a bracket. */
+	virtual void NotifyBracketOpen() = 0;
+	virtual void NotifyBracketClosed() = 0;
+
+private:
+	friend class FAnimDataControllerTestBase;
+	friend UE::Anim::FOpenBracketAction;
+	friend UE::Anim::FCloseBracketAction;
+};
diff --git a/Engine/Source/Runtime/Engine/Classes/Animation/AnimSequence.h b/Engine/Source/Runtime/Engine/Classes/Animation/AnimSequence.h
--- a/Engine/Source/Runtime/Engine/Classes/Animation/AnimSequence.h
+++ b/Engine/Source/Runtime/Engine/Classes/Animation/AnimSequence.h
@@ -24,7 +24,6 @@
 
 #if WITH_EDITOR
 #include "AnimData/AnimDataModel.h"
-#include "AnimData/AnimDataController.h"
 #endif // WITH_EDITOR
 
 #include "AnimSequence.generated.h"
diff --git a/Engine/Source/Runtime/Engine/Classes/Animation/AnimSequenceBase.h b/Engine/Source/Runtime/Engine/Classes/Animation/AnimSequenceBase.h
--- a/Engine/Source/Runtime/Engine/Classes/Animation/AnimSequenceBase.h
+++ b/Engine/Source/Runtime/Engine/Classes/Animation/AnimSequenceBase.h
@@ -14,6 +14,7 @@
 #include "Animation/AnimCurveTypes.h"
 #include "Animation/AnimNotifyQueue.h"
 #include "Animation/AnimData/AnimDataModelNotifyCollector.h"
+#include "Animation/AnimData/IAnimationDataController.h"
 
 #include "AnimSequenceBase.generated.h"
 
@@ -280,8 +281,7 @@ public:
 	UAnimDataModel* GetDataModel() const;
 
 	/** Returns the transient UAnimDataController set to operate on DataModel */
-	UAnimDataController* GetController();
-
+	IAnimationDataController& GetController();
 protected:
 	/** Populates the UAnimDataModel object according to any pre-existing data. (overrides expect to populate the model according to their data) */
 	virtual void PopulateModel();
@@ -310,7 +310,7 @@ protected:
 
 	/** UAnimDataController instance set to operate on DataModel */
 	UPROPERTY(VisibleAnywhere, BlueprintReadOnly, Transient, DuplicateTransient, Category = "Animation Model")
-	TObjectPtr<UAnimDataController> Controller;
+	TScriptInterface<IAnimationDataController> Controller;
 	
 	/** Helper object that keeps track of any controller brackets, and all unique notify types that are broadcasted during it */
 	UE::Anim::FAnimDataModelNotifyCollector NotifyCollector;
diff --git a/Engine/Source/Runtime/Engine/Classes/Animation/AnimSequenceHelpers.h b/Engine/Source/Runtime/Engine/Classes/Animation/AnimSequenceHelpers.h
--- a/Engine/Source/Runtime/Engine/Classes/Animation/AnimSequenceHelpers.h
+++ b/Engine/Source/Runtime/Engine/Classes/Animation/AnimSequenceHelpers.h
@@ -83,7 +83,7 @@ namespace Anim {
 	* @param	Skeleton			Skeleton to use for verifying and or add the curve names
 	* @param	Controller			Controller to use for adding curve data to a model
 	*/
-	ENGINE_API void CopyCurveDataToModel(const FRawCurveTracks& CurveData, const USkeleton* Skeleton, UAnimDataController* Controller);
+	ENGINE_API void CopyCurveDataToModel(const FRawCurveTracks& CurveData, const USkeleton* Skeleton, IAnimationDataController& Controller);
 
 	/**
 	* Copies over any individual curve from FRawCurveTracks to a UAnimDataModel instance targeted by the provided controller.
diff --git a/Engine/Source/Runtime/Engine/Engine.Build.cs b/Engine/Source/Runtime/Engine/Engine.Build.cs
--- a/Engine/Source/Runtime/Engine/Engine.Build.cs
+++ b/Engine/Source/Runtime/Engine/Engine.Build.cs
@@ -230,6 +230,9 @@ public class Engine : ModuleRules
 			PrivateIncludePathModuleNames.Add("MeshUtilitiesCommon");
 
 			DynamicallyLoadedModuleNames.Add("MeshUtilities");
+			
+			PublicIncludePathModuleNames.Add("AnimationDataController");
+			DynamicallyLoadedModuleNames.Add("AnimationDataController");
 
 			PrivateDependencyModuleNames.AddRange(
 				new string[] {
diff --git a/Engine/Source/Runtime/Engine/Private/Animation/AnimComposite.cpp b/Engine/Source/Runtime/Engine/Private/Animation/AnimComposite.cpp
--- a/Engine/Source/Runtime/Engine/Private/Animation/AnimComposite.cpp
+++ b/Engine/Source/Runtime/Engine/Private/Animation/AnimComposite.cpp
@@ -6,8 +6,7 @@
 
 #include "Animation/AnimComposite.h"
 #include "Animation/AnimationPoseData.h"
-#include "Animation/AttributesRuntime.h"
-#include "Animation/AnimData/AnimDataController.h"
+#include "Animation/CustomAttributesRuntime.h"
 
 UAnimComposite::UAnimComposite(const FObjectInitializer& ObjectInitializer)
 	: Super(ObjectInitializer)
diff --git a/Engine/Source/Runtime/Engine/Private/Animation/AnimData/AnimDataControllerActions.h b/Engine/Source/Runtime/Engine/Private/Animation/AnimData/AnimDataControllerActions.h
deleted file mode 100644
diff --git a/Engine/Source/Runtime/Engine/Private/Animation/AnimDataController.cpp b/Engine/Source/Runtime/Engine/Private/Animation/AnimDataController.cpp
deleted file mode 100644
diff --git a/Engine/Source/Runtime/Engine/Private/Animation/AnimDataControllerActions.cpp b/Engine/Source/Runtime/Engine/Private/Animation/AnimDataControllerActions.cpp
deleted file mode 100644
diff --git a/Engine/Source/Runtime/Engine/Private/Animation/AnimMontage.cpp b/Engine/Source/Runtime/Engine/Private/Animation/AnimMontage.cpp
--- a/Engine/Source/Runtime/Engine/Private/Animation/AnimMontage.cpp
+++ b/Engine/Source/Runtime/Engine/Private/Animation/AnimMontage.cpp
@@ -19,7 +19,6 @@
 #include "Animation/AnimSingleNodeInstance.h"
 #include "Engine/Engine.h"
 #include "Animation/AnimTrace.h"
-#include "Animation/AnimData/AnimDataController.h"
 
 DEFINE_LOG_CATEGORY(LogAnimMontage);
 
diff --git a/Engine/Source/Runtime/Engine/Private/Animation/AnimSequence.cpp b/Engine/Source/Runtime/Engine/Private/Animation/AnimSequence.cpp
--- a/Engine/Source/Runtime/Engine/Private/Animation/AnimSequence.cpp
+++ b/Engine/Source/Runtime/Engine/Private/Animation/AnimSequence.cpp
@@ -57,7 +57,6 @@
 #include "HAL/FileManager.h"
 
 #if WITH_EDITOR
-#include "Animation/AnimData/AnimDataController.h"
 #include "Animation/AnimData/AnimDataModel.h"
 #include "Animation/BuiltInAttributeTypes.h"
 #endif // WITH_EDITOR
@@ -1880,7 +1879,7 @@ bool UAnimSequence::InsertFramesToRawAnimData( int32 StartFrame, int32 EndFrame,
 	const int32 NumFramesToInsert = EndFrame-StartFrame;
 	if ((CopyFrame>=0 && CopyFrame< DataModel->GetNumberOfKeys()) && (StartFrame >= 0 && StartFrame <= DataModel->GetNumberOfKeys()) && NumFramesToInsert > 0)
 	{
-		UAnimDataController::FScopedBracket ScopedBracket(Controller, LOCTEXT("InsertFramesToRawAnimData", "Inserting Frames into Animation Track Data"));
+		IAnimationDataController::FScopedBracket ScopedBracket(Controller, LOCTEXT("InsertFramesToRawAnimData", "Inserting Frames into Animation Track Data"));
 		UE::Anim::AnimationData::DuplicateKeys(this, StartFrame, NumFramesToInsert, CopyFrame);
 
 		UE_LOG(LogAnimation, Log, TEXT("\tPlay Length: %f, Number of Keys: %d"), GetPlayLength(), DataModel->GetNumberOfKeys());
@@ -2859,8 +2858,7 @@ void UAnimSequence::RemapTracksToNewSkeleton( USkeleton* NewSkeleton, bool bConv
 
 	ValidateModel();
 
-	UAnimDataController::FScopedBracket ScopedBracket(Controller, LOCTEXT("RemapTrackToNewSkeleton_Bracket", "Remapping Bone Animation tracks to new Skeleton"));
-	
+	IAnimationDataController::FScopedBracket ScopedBracket(Controller, LOCTEXT("RemapTrackToNewSkeleton_Bracket", "Remapping Bone Animation tracks to new Skeleton"));
 	Controller->UpdateAttributesFromSkeleton(NewSkeleton);
 
 	// @Todo : currently additive will work fine since we don't bake anything except when we extract
@@ -3335,7 +3333,7 @@ void UAnimSequence::RemoveNaNTracks()
 		}
 	}
 
-	UAnimDataController::FScopedBracket ScopedBracket(Controller, LOCTEXT("RemoveNaNTracks_Bracket", "Removing track(s) containing NaN key data"));	
+	IAnimationDataController::FScopedBracket ScopedBracket(Controller, LOCTEXT("RemoveNaNTracks_Bracket", "Removing track(s) containing NaN key data"));	
 	for (const FName& TrackName : TracksToRemove)
 	{
 		Controller->RemoveBoneTrack(TrackName);
@@ -3352,7 +3350,7 @@ void UAnimSequence::RemoveAllTracks()
 	PRAGMA_ENABLE_DEPRECATION_WARNINGS
 	
 	ValidateModel();
-	UAnimDataController::FScopedBracket ScopedBracket(Controller, LOCTEXT("RemoveAllTracks_Bracket", "Removing all Bone Animation and Transform Curve Tracks"));
+	IAnimationDataController::FScopedBracket ScopedBracket(Controller, LOCTEXT("RemoveAllTracks_Bracket", "Removing all Bone Animation and Transform Curve Tracks"));
 	Controller->RemoveAllBoneTracks();
 
 	// clear all transform tracks
@@ -3478,7 +3476,7 @@ void UAnimSequence::ResizeSequence(float NewLength, int32 NewNumFrames, bool bIn
 
 	FFrameRate NewFrameRate = FFrameRate((NewLength / NewNumFrames), 1);
 
-	GetController()->Resize(NewLength, T0, T1);
+	Controller->Resize(NewLength, T0, T1);
 
 	ensure(DataModel->GetNumberOfKeys() == NewNumFrames);
 }
@@ -3885,7 +3883,7 @@ bool UAnimSequence::ConvertRiggingDataToAnimationData(FAnimSequenceTrackContaine
 	{
 		ValidateModel();
 
-		UAnimDataController::FScopedBracket ScopedBracket(Controller, LOCTEXT("ConvertRiggingDataToAnimationData_Bracket", "Converting Rigging Data to Animation Data"));
+		IAnimationDataController::FScopedBracket ScopedBracket(Controller, LOCTEXT("ConvertRiggingDataToAnimationData_Bracket", "Converting Rigging Data to Animation Data"));
 
 		TArray< TArray<FTransform> > AnimationDataInComponentSpace;
 		int32 NumBones = GetSpaceBasedAnimationData(AnimationDataInComponentSpace, &RiggingAnimationData);
@@ -4128,7 +4126,7 @@ void UAnimSequence::AddKeyToSequence(float Time, const FName& BoneName, const FT
 
 	ValidateModel();
 
-	UAnimDataController::FScopedBracket ScopedBracket(Controller, LOCTEXT("AddKeyToSequence_Bracket", "Adding key to sequence"));
+	IAnimationDataController::FScopedBracket ScopedBracket(Controller, LOCTEXT("AddKeyToSequence_Bracket", "Adding key to sequence"));
 	FAnimationCurveIdentifier TransformCurveId(NewCurveName, ERawCurveTrackTypes::RCT_Transform);
 	Controller->AddCurve(TransformCurveId, AACF_DriveTrack | AACF_Editable);
 
@@ -4166,7 +4164,7 @@ bool UAnimSequence::CreateAnimation(USkeletalMesh* Mesh)
 		ValidateModel();
 		const FReferenceSkeleton& RefSkeleton = Mesh->GetRefSkeleton();
 
-		UAnimDataController::FScopedBracket ScopedBracket(Controller, LOCTEXT("CreateAnimationSkeletalMesh_Bracket", "Creating Animation Sequence based up Skeletal Mesh"));
+		IAnimationDataController::FScopedBracket ScopedBracket(Controller, LOCTEXT("CreateAnimationSkeletalMesh_Bracket", "Creating Animation Sequence based up Skeletal Mesh"));
 		ResetAnimation();
 
 		Controller->SetPlayLength(MINIMUM_ANIMATION_LENGTH);
@@ -4203,7 +4201,7 @@ bool UAnimSequence::CreateAnimation(USkeletalMeshComponent* MeshComponent)
 		ValidateModel();
 		const FReferenceSkeleton& RefSkeleton = Mesh->GetRefSkeleton();
 
-		UAnimDataController::FScopedBracket ScopedBracket(Controller, LOCTEXT("CreateAnimationSkeletalMeshComponent_Bracket", "Creating Animation Sequence based up Skeletal Mesh Component"));
+		IAnimationDataController::FScopedBracket ScopedBracket(Controller, LOCTEXT("CreateAnimationSkeletalMeshComponent_Bracket", "Creating Animation Sequence based up Skeletal Mesh Component"));
 		ResetAnimation();
 
 		Controller->SetPlayLength(MINIMUM_ANIMATION_LENGTH);
@@ -5492,7 +5490,7 @@ void UAnimSequence::PopulateModel()
 	Controller->SetFrameRate(FrameRate);
 
 	USkeleton* TargetSkeleton = GetSkeleton();
-	UE::Anim::CopyCurveDataToModel(CurveData, TargetSkeleton, Controller);
+	UE::Anim::CopyCurveDataToModel(CurveData, TargetSkeleton, *Controller);
 	
 	const int32 NumTracks = SequenceTracks.Num();
 	for (int32 TrackIndex = 0; TrackIndex < NumTracks; ++TrackIndex)
diff --git a/Engine/Source/Runtime/Engine/Private/Animation/AnimSequenceBase.cpp b/Engine/Source/Runtime/Engine/Private/Animation/AnimSequenceBase.cpp
--- a/Engine/Source/Runtime/Engine/Private/Animation/AnimSequenceBase.cpp
+++ b/Engine/Source/Runtime/Engine/Private/Animation/AnimSequenceBase.cpp
@@ -13,7 +13,10 @@
 #include "Animation/AnimationPoseData.h"
 #include "Animation/AnimSequenceHelpers.h"
 #include "Animation/AnimData/AnimDataModel.h"
-#include "Animation/AnimData/AnimDataController.h"
+
+#if WITH_EDITOR
+#include "IAnimationDataControllerModule.h"
+#endif // WITH_EDITOR
 
 #include "UObject/UE5MainStreamObjectVersion.h"
 
@@ -45,7 +48,7 @@ void UAnimSequenceBase::SetSequenceLength(float NewLength)
 {
 #if WITH_EDITOR
 	// Editor time we update the source data
-	GetController()->SetPlayLength(NewLength);
+	Controller->SetPlayLength(NewLength);
 #else 
 	// In case this is called during runtime, update the sequence length directly. Current use-case is runtime created Montages.
 	PRAGMA_DISABLE_DEPRECATION_WARNINGS
@@ -1126,25 +1129,20 @@ UAnimDataModel* UAnimSequenceBase::GetDataModel() const
 	return DataModel;
 }
 
-UAnimDataController* UAnimSequenceBase::GetController()
+IAnimationDataController& UAnimSequenceBase::GetController()
 {
 	ValidateModel();
 
-	if (HasAllFlags(EObjectFlags::RF_ClassDefaultObject))
+	if(Controller == nullptr)
 	{
-		return nullptr;
+		IAnimationDataControllerModule& ControllerModule = FModuleManager::Get().LoadModuleChecked<IAnimationDataControllerModule>("AnimationDataController");
+		Controller = ControllerModule.GetController();
+		Controller->SetModel(DataModel);
 	}
 
-	if (Controller == nullptr)
-	{
-		Controller = NewObject<UAnimDataController>(GetTransientPackage());		
-	}
+	ensure(Controller->GetModel() == DataModel);
 
-	if (Controller->GetModel() != DataModel)
-	{
-		Controller->SetModel(DataModel);
-	}
-	return Controller;
+	return *Controller;
 }
 
 void UAnimSequenceBase::PopulateModel()
@@ -1158,12 +1156,12 @@ void UAnimSequenceBase::PopulateModel()
 	const FRawCurveTracks CurveData = RawCurveData;
 	PRAGMA_ENABLE_DEPRECATION_WARNINGS
 		
-	UAnimDataController::FScopedBracket ScopedBracket(Controller, LOCTEXT("UAnimSequenceBase::PopulateModel_Bracket", "Generating Animation Model Data for Animation Sequence Base"));
+	IAnimationDataController::FScopedBracket ScopedBracket(Controller, LOCTEXT("UAnimSequenceBase::PopulateModel_Bracket", "Generating Animation Model Data for Animation Sequence Base"));
 	Controller->SetPlayLength(PlayLength);
 	Controller->SetFrameRate(FrameRate);
 
 	USkeleton* TargetSkeleton = GetSkeleton();	
-	UE::Anim::CopyCurveDataToModel(CurveData, TargetSkeleton, Controller);	
+	UE::Anim::CopyCurveDataToModel(CurveData, TargetSkeleton,  *Controller);	
 }
 
 void UAnimSequenceBase::ValidateModel() const
@@ -1180,7 +1178,7 @@ void UAnimSequenceBase::CopyDataModel(const UAnimDataModel* ModelToDuplicate)
 	}
 
 	DataModel = DuplicateObject(ModelToDuplicate, this);
-	GetController()->SetModel(DataModel);
+	Controller->SetModel(DataModel);
 	DataModel->GetModifiedEvent().AddUObject(this, &UAnimSequenceBase::OnModelModified);
 }
 
@@ -1190,8 +1188,6 @@ void UAnimSequenceBase::CreateModel()
 	{
 		DataModel = CreateDefaultSubobject<UAnimDataModel>(FName(TEXT("AnimationDataModel")));
 	}
-
-	GetController()->SetModel(DataModel);
 	
 	DataModel->GetModifiedEvent().RemoveAll(this);
 	DataModel->GetModifiedEvent().AddUObject(this, &UAnimSequenceBase::OnModelModified);
diff --git a/Engine/Source/Runtime/Engine/Private/Animation/AnimSequenceHelpers.cpp b/Engine/Source/Runtime/Engine/Private/Animation/AnimSequenceHelpers.cpp
--- a/Engine/Source/Runtime/Engine/Private/Animation/AnimSequenceHelpers.cpp
+++ b/Engine/Source/Runtime/Engine/Private/Animation/AnimSequenceHelpers.cpp
@@ -308,7 +308,7 @@ void GetBoneTransformFromModel(const UAnimDataModel* Model, FTransform& OutTrans
 	}
 }
 
-void CopyCurveDataToModel(const FRawCurveTracks& CurveData, const USkeleton* Skeleton, UAnimDataController* Controller)
+void CopyCurveDataToModel(const FRawCurveTracks& CurveData, const USkeleton* Skeleton, IAnimationDataController& Controller)
 {
 	// Populate float curve data
 	for (const FFloatCurve& FloatCurve : CurveData.FloatCurves)
@@ -316,8 +316,8 @@ void CopyCurveDataToModel(const FRawCurveTracks& CurveData, const USkeleton* Ske
 		const FAnimationCurveIdentifier CurveId = UAnimationCurveIdentifierExtensions::FindCurveIdentifier(Skeleton, FloatCurve.Name.DisplayName, ERawCurveTrackTypes::RCT_Float);
 		if (CurveId.IsValid())
 		{
-			Controller->AddCurve(CurveId, FloatCurve.GetCurveTypeFlags());
-			Controller->SetCurveKeys(CurveId, FloatCurve.FloatCurve.GetConstRefOfKeys());
+			Controller.AddCurve(CurveId, FloatCurve.GetCurveTypeFlags());
+			Controller.SetCurveKeys(CurveId, FloatCurve.FloatCurve.GetConstRefOfKeys());
 		}
 	}
 
@@ -327,7 +327,7 @@ void CopyCurveDataToModel(const FRawCurveTracks& CurveData, const USkeleton* Ske
 		const FAnimationCurveIdentifier CurveId = UAnimationCurveIdentifierExtensions::FindCurveIdentifier(Skeleton, TransformCurve.Name.DisplayName, ERawCurveTrackTypes::RCT_Transform);
 		if (CurveId.IsValid())
 		{
-			Controller->AddCurve(CurveId, TransformCurve.GetCurveTypeFlags());
+			Controller.AddCurve(CurveId, TransformCurve.GetCurveTypeFlags());
 
 			// Set each individual channel rich curve keys, to account for any custom tangents etc.
 			for (int32 SubCurveIndex = 0; SubCurveIndex < 3; ++SubCurveIndex)
@@ -339,7 +339,7 @@ void CopyCurveDataToModel(const FRawCurveTracks& CurveData, const USkeleton* Ske
 					const EVectorCurveChannel Axis = (EVectorCurveChannel)ChannelIndex;
 					FAnimationCurveIdentifier TargetCurveIdentifier = CurveId;
 					UAnimationCurveIdentifierExtensions::GetTransformChildCurveIdentifier(TargetCurveIdentifier, Channel, Axis);
-					Controller->SetCurveKeys(TargetCurveIdentifier, VectorCurve->FloatCurves[ChannelIndex].GetConstRefOfKeys());
+					Controller.SetCurveKeys(TargetCurveIdentifier, VectorCurve->FloatCurves[ChannelIndex].GetConstRefOfKeys());
 				}
 			}
 		}
@@ -707,7 +707,7 @@ void Compression::SanitizeRawAnimSequenceTrack(FRawAnimSequenceTrack& RawTrack)
 bool AnimationData::AddLoopingInterpolation(UAnimSequence* InSequence)
 {
 	const UAnimDataModel* DataModel = InSequence->GetDataModel();
-	UAnimDataController* Controller = InSequence->GetController();
+	IAnimationDataController& Controller = InSequence->GetController();
 
 	const int32 NumTracks = DataModel->GetNumBoneTracks();
 	const int32 NumKeys = DataModel->GetNumberOfKeys();
@@ -725,7 +725,7 @@ bool AnimationData::AddLoopingInterpolation(UAnimSequence* InSequence)
 			}
 		};
 		
-		UAnimDataController::FScopedBracket ScopedBracket(Controller, LOCTEXT("AddLoopingInterpolation_Bracket", "Adding looping interpolation"));
+		IAnimationDataController::FScopedBracket ScopedBracket(Controller, LOCTEXT("AddLoopingInterpolation_Bracket", "Adding looping interpolation"));
 
 		const TArray<FBoneAnimationTrack>& BoneAnimationTracks = DataModel->GetBoneAnimationTracks();
 		for (const FBoneAnimationTrack& AnimationTrack : BoneAnimationTracks)
@@ -739,11 +739,11 @@ bool AnimationData::AddLoopingInterpolation(UAnimSequence* InSequence)
 			auto ScaleKeys = AnimationTrack.InternalTrackData.ScaleKeys;
 			LoopKeyData(ScaleKeys);
 
-			Controller->SetBoneTrackKeys(AnimationTrack.Name, PositionalKeys, RotationalKeys, ScaleKeys);
+			Controller.SetBoneTrackKeys(AnimationTrack.Name, PositionalKeys, RotationalKeys, ScaleKeys);
 		}
 
 		const float Interval = DataModel->GetFrameRate().AsInterval();
-		Controller->SetPlayLength(Interval + DataModel->GetPlayLength());
+		Controller.SetPlayLength(Interval + DataModel->GetPlayLength());
 
 		return true;
 	}
@@ -754,7 +754,7 @@ bool AnimationData::AddLoopingInterpolation(UAnimSequence* InSequence)
 bool AnimationData::Trim(UAnimSequence* InSequence, float TrimStart, float TrimEnd)
 {
 	const UAnimDataModel* DataModel = InSequence->GetDataModel();
-	UAnimDataController* Controller = InSequence->GetController();
+	IAnimationDataController& Controller = InSequence->GetController();
 
 	const int32 NumTracks = DataModel->GetNumBoneTracks();
 	const int32 NumKeys = DataModel->GetNumberOfKeys();
@@ -778,7 +778,7 @@ bool AnimationData::Trim(UAnimSequence* InSequence, float TrimStart, float TrimE
 		const int32 StartTrimKeyIndex = StartFrameTrim.Value + 1;
 		const int32 NumTrimmedFrames = EndFrameTrim.Value - StartFrameTrim.Value;
 		
-		UAnimDataController::FScopedBracket ScopedBracket(Controller, LOCTEXT("TrimRawAnimation_Bracket", "Trimming Animation Track Data"));
+		IAnimationDataController::FScopedBracket ScopedBracket(Controller, LOCTEXT("TrimRawAnimation_Bracket", "Trimming Animation Track Data"));
 		RemoveKeys(InSequence, StartTrimKeyIndex, NumTrimmedFrames);
 		
 		return true;
@@ -790,7 +790,7 @@ bool AnimationData::Trim(UAnimSequence* InSequence, float TrimStart, float TrimE
 void AnimationData::DuplicateKeys(UAnimSequence* InSequence, int32 StartKeyIndex, int32 NumDuplicates, int32 SourceKeyIndex /*= INDEX_NONE */)
 {
 	const UAnimDataModel* Model = InSequence->GetDataModel();
-	UAnimDataController* Controller = InSequence->GetController();
+	IAnimationDataController& Controller = InSequence->GetController();
 
 	const int32 NumberOfKeys = Model->GetNumberOfKeys();
 	const FFrameRate& FrameRate = Model->GetFrameRate();
@@ -801,7 +801,7 @@ void AnimationData::DuplicateKeys(UAnimSequence* InSequence, int32 StartKeyIndex
 		const int32 CopyKeyIndex = SourceKeyIndex == INDEX_NONE ? StartKeyIndex : SourceKeyIndex;
 		if (CopyKeyIndex >= 0 && CopyKeyIndex < NumberOfKeys)
 		{
-			UAnimDataController::FScopedBracket ScopedBracket(Controller, LOCTEXT("DuplicateKeys_Bracket", "Inserting Duplicate Animation Track Keys"));
+			IAnimationDataController::FScopedBracket ScopedBracket(Controller, LOCTEXT("DuplicateKeys_Bracket", "Inserting Duplicate Animation Track Keys"));
 
 			const int32 NumFramesToInsert = NumDuplicates;
 			const int32 EndFrameIndex = StartKeyIndex + NumDuplicates;
@@ -829,7 +829,7 @@ void AnimationData::DuplicateKeys(UAnimSequence* InSequence, int32 StartKeyIndex
 				InsertFrames(TrackData.RotKeys);
 				InsertFrames(TrackData.ScaleKeys);
 
-				Controller->SetBoneTrackKeys(AnimationTrack.Name, TrackData.PosKeys, TrackData.RotKeys, TrackData.ScaleKeys);
+				Controller.SetBoneTrackKeys(AnimationTrack.Name, TrackData.PosKeys, TrackData.RotKeys, TrackData.ScaleKeys);
 			}
 
 			// The number of keys has changed, which means that the sequence length and number of frames should be updated as well
@@ -841,7 +841,7 @@ void AnimationData::DuplicateKeys(UAnimSequence* InSequence, int32 StartKeyIndex
 			const float InsertedTime = FrameRate.AsInterval() * NumDuplicates;
 
 			// Notify will happen with time slice that was inserted
-			Controller->Resize(NewSequenceLength, StartTime, StartTime + InsertedTime);
+			Controller.Resize(NewSequenceLength, StartTime, StartTime + InsertedTime);
 		}
 	}
 }
@@ -849,7 +849,7 @@ void AnimationData::DuplicateKeys(UAnimSequence* InSequence, int32 StartKeyIndex
 void AnimationData::RemoveKeys(UAnimSequence* InSequence, int32 StartKeyIndex, int32 NumKeysToRemove)
 {
 	const UAnimDataModel* Model = InSequence->GetDataModel();
-	UAnimDataController* Controller = InSequence->GetController();
+	IAnimationDataController& Controller = InSequence->GetController();
 
 	const int32 NumberOfKeys = Model->GetNumberOfKeys();
 	const FFrameRate& FrameRate = Model->GetFrameRate();
@@ -857,7 +857,7 @@ void AnimationData::RemoveKeys(UAnimSequence* InSequence, int32 StartKeyIndex, i
 	const int32 EndKeyIndex = StartKeyIndex + NumKeysToRemove;
 	if (StartKeyIndex >= 0 && StartKeyIndex < NumberOfKeys && NumKeysToRemove > 0 && EndKeyIndex <= NumberOfKeys)
 	{
-		UAnimDataController::FScopedBracket ScopedBracket(Controller, LOCTEXT("RemoveKeys_Bracket", "Removing Animation Track Keys"));
+		IAnimationDataController::FScopedBracket ScopedBracket(Controller, LOCTEXT("RemoveKeys_Bracket", "Removing Animation Track Keys"));
 
 		auto ShrinkKeys = [&](auto& KeyData)
 		{
@@ -880,7 +880,7 @@ void AnimationData::RemoveKeys(UAnimSequence* InSequence, int32 StartKeyIndex, i
 			ShrinkKeys(TrackData.RotKeys);
 			ShrinkKeys(TrackData.ScaleKeys);
 			
-			Controller->SetBoneTrackKeys(AnimationTrack.Name, TrackData.PosKeys, TrackData.RotKeys, TrackData.ScaleKeys);
+			Controller.SetBoneTrackKeys(AnimationTrack.Name, TrackData.PosKeys, TrackData.RotKeys, TrackData.ScaleKeys);
 		}
 
 		const int32 NewNumberOfFrames = FMath::Max(NewNumberOfKeys - 1, 0);
@@ -891,7 +891,7 @@ void AnimationData::RemoveKeys(UAnimSequence* InSequence, int32 StartKeyIndex, i
 		const float EndTime = StartFrameTime + RemovedTime;
 
 		// Notify will happen with time slice that was removed
-		Controller->Resize(NewSequenceLength, StartFrameTime, EndTime);
+		Controller.Resize(NewSequenceLength, StartFrameTime, EndTime);
 	}
 }
 
diff --git a/Engine/Source/Runtime/Engine/Private/Animation/AnimStreamable.cpp b/Engine/Source/Runtime/Engine/Private/Animation/AnimStreamable.cpp
--- a/Engine/Source/Runtime/Engine/Private/Animation/AnimStreamable.cpp
+++ b/Engine/Source/Runtime/Engine/Private/Animation/AnimStreamable.cpp
@@ -27,7 +27,6 @@
 
 #include "Animation/AnimData/AnimDataModel.h"
 #include "Animation/AnimSequenceHelpers.h"
-#include "Animation/AnimData/AnimDataController.h"
 
 CSV_DECLARE_CATEGORY_MODULE_EXTERN(ENGINE_API, Animation);
 
diff --git a/Engine/Source/Runtime/Engine/Private/AnimationAsset.cpp b/Engine/Source/Runtime/Engine/Private/AnimationAsset.cpp
--- a/Engine/Source/Runtime/Engine/Private/AnimationAsset.cpp
+++ b/Engine/Source/Runtime/Engine/Private/AnimationAsset.cpp
@@ -11,7 +11,6 @@
 #include "Animation/BlendSpace.h"
 #include "Animation/PoseAsset.h"
 #include "Animation/AnimSequenceHelpers.h"
-#include "Animation/AnimData/AnimDataController.h"
 
 #define LOCTEXT_NAMESPACE "AnimationAsset"
 
@@ -409,7 +408,7 @@ bool UAnimationAsset::ReplaceSkeleton(USkeleton* NewSkeleton, bool bConvertSpace
 				// raw animation data itself.
 				if (UAnimSequence* Sequence = Cast<UAnimSequence>(AnimAsset))
 				{
-					Sequence->GetController()->OpenBracket(LOCTEXT("ReplaceSkeleton_Bracket", "Replacing USkeleton"));
+					Sequence->GetController().OpenBracket(LOCTEXT("ReplaceSkeleton_Bracket", "Replacing USkeleton"));
 				}
 
 				// these two are different functions for now
@@ -422,7 +421,7 @@ bool UAnimationAsset::ReplaceSkeleton(USkeleton* NewSkeleton, bool bConvertSpace
 			{
 				if (UAnimSequence* Seq = Cast<UAnimSequence>(AnimAsset))
 				{
-					Seq->GetController()->CloseBracket();
+					Seq->GetController().CloseBracket();
 				}
 			}
 		}
@@ -433,14 +432,14 @@ bool UAnimationAsset::ReplaceSkeleton(USkeleton* NewSkeleton, bool bConvertSpace
 			// raw animation data itself.
 			if (Seq)
 			{
-				Seq->GetController()->OpenBracket(LOCTEXT("ReplaceSkeleton_Bracket", "Replacing USkeleton"));
+				Seq->GetController().OpenBracket(LOCTEXT("ReplaceSkeleton_Bracket", "Replacing USkeleton"));
 			}
   
 			RemapTracksToNewSkeleton(NewSkeleton, bConvertSpaces);
 
 			if (Seq)
 			{
-				Seq->GetController()->CloseBracket();
+				Seq->GetController().CloseBracket();
 			}
 		}
 
diff --git a/Engine/Source/Runtime/Engine/Public/Animation/AnimCurveTypes.h b/Engine/Source/Runtime/Engine/Public/Animation/AnimCurveTypes.h
--- a/Engine/Source/Runtime/Engine/Public/Animation/AnimCurveTypes.h
+++ b/Engine/Source/Runtime/Engine/Public/Animation/AnimCurveTypes.h
@@ -157,7 +157,7 @@ private:
 };
 
 USTRUCT(BlueprintType)
-struct FFloatCurve : public FAnimCurveBase
+struct ENGINE_API FFloatCurve : public FAnimCurveBase
 {
 	GENERATED_USTRUCT_BODY()
 
@@ -174,14 +174,14 @@ struct FFloatCurve : public FAnimCurveBase
 
 	// we don't want to have = operator. This only copies curves, but leaving naming and everything else intact. 
 	void CopyCurve(const FFloatCurve& SourceCurve);
-	ENGINE_API float Evaluate(float CurrentTime) const;
-	ENGINE_API void UpdateOrAddKey(float NewKey, float CurrentTime);
-	ENGINE_API void GetKeys(TArray<float>& OutTimes, TArray<float>& OutValues) const;
+	float Evaluate(float CurrentTime) const;
+	void UpdateOrAddKey(float NewKey, float CurrentTime);
+	void GetKeys(TArray<float>& OutTimes, TArray<float>& OutValues) const;
 	void Resize(float NewLength, bool bInsert/* whether insert or remove*/, float OldStartTime, float OldEndTime);
 };
 
 USTRUCT()
-struct FVectorCurve : public FAnimCurveBase
+struct ENGINE_API FVectorCurve : public FAnimCurveBase
 {
 	GENERATED_USTRUCT_BODY()
 
@@ -206,16 +206,16 @@ struct FVectorCurve : public FAnimCurveBase
 
 	// we don't want to have = operator. This only copies curves, but leaving naming and everything else intact. 
 	void CopyCurve(const FVectorCurve& SourceCurve);
-	ENGINE_API FVector Evaluate(float CurrentTime, float BlendWeight) const;
-	ENGINE_API void UpdateOrAddKey(const FVector& NewKey, float CurrentTime);
-	ENGINE_API void GetKeys(TArray<float>& OutTimes, TArray<FVector>& OutValues) const;
+	FVector Evaluate(float CurrentTime, float BlendWeight) const;
+	void UpdateOrAddKey(const FVector& NewKey, float CurrentTime);
+	void GetKeys(TArray<float>& OutTimes, TArray<FVector>& OutValues) const;
 	bool DoesContainKey() const { return (FloatCurves[0].GetNumKeys() > 0 || FloatCurves[1].GetNumKeys() > 0 || FloatCurves[2].GetNumKeys() > 0);}
 	void Resize(float NewLength, bool bInsert/* whether insert or remove*/, float OldStartTime, float OldEndTime);
 	int32 GetNumKeys() const;
 };
 
 USTRUCT(BlueprintType)
-struct FTransformCurve: public FAnimCurveBase
+struct ENGINE_API FTransformCurve: public FAnimCurveBase
 {
 	GENERATED_USTRUCT_BODY()
 
@@ -242,9 +242,9 @@ struct FTransformCurve: public FAnimCurveBase
 
 	// we don't want to have = operator. This only copies curves, but leaving naming and everything else intact. 
 	void CopyCurve(const FTransformCurve& SourceCurve);
-	ENGINE_API FTransform Evaluate(float CurrentTime, float BlendWeight) const;
-	ENGINE_API void UpdateOrAddKey(const FTransform& NewKey, float CurrentTime);
-	ENGINE_API void GetKeys(TArray<float>& OutTimes, TArray<FTransform>& OutValues) const;
+	FTransform Evaluate(float CurrentTime, float BlendWeight) const;
+	void UpdateOrAddKey(const FTransform& NewKey, float CurrentTime);
+	void GetKeys(TArray<float>& OutTimes, TArray<FTransform>& OutValues) const;
 	void Resize(float NewLength, bool bInsert/* whether insert or remove*/, float OldStartTime, float OldEndTime);
 	
 	const FVectorCurve* GetVectorCurveByIndex(int32 Index) const;
