Non-unity fixes

diff --git a/Engine/Source/Developer/AnimationDataController/Private/AnimationDataControllerModule.cpp b/Engine/Source/Developer/AnimationDataController/Private/AnimationDataControllerModule.cpp
--- a/Engine/Source/Developer/AnimationDataController/Private/AnimationDataControllerModule.cpp
+++ b/Engine/Source/Developer/AnimationDataController/Private/AnimationDataControllerModule.cpp
@@ -2,7 +2,7 @@
 
 #include "IAnimationDataControllerModule.h"
 #include "AnimDataController.h"
-#include "UObject/StrongObjectPtr.h"
+#include "Modules/ModuleManager.h"
 
 class FAnimationDataControllerModule : public IAnimationDataControllerModule
 {
diff --git a/Engine/Source/Editor/Persona/Private/AnimTimeline/AnimTimelineTrack_FloatCurve.cpp b/Engine/Source/Editor/Persona/Private/AnimTimeline/AnimTimelineTrack_FloatCurve.cpp
--- a/Engine/Source/Editor/Persona/Private/AnimTimeline/AnimTimelineTrack_FloatCurve.cpp
+++ b/Engine/Source/Editor/Persona/Private/AnimTimeline/AnimTimelineTrack_FloatCurve.cpp
@@ -21,6 +21,7 @@
 #include "AnimModel_AnimSequenceBase.h"
 #include "Animation/AnimData/AnimDataModel.h"
 #include "SAnimOutlinerItem.h"
+#include "Animation/AnimData/AnimDataModel.h"
 
 #define LOCTEXT_NAMESPACE "FAnimTimelineTrack_FloatCurve"
 
diff --git a/Engine/Source/Runtime/Engine/Classes/Animation/AnimData/IAnimationDataController.h b/Engine/Source/Runtime/Engine/Classes/Animation/AnimData/IAnimationDataController.h
--- a/Engine/Source/Runtime/Engine/Classes/Animation/AnimData/IAnimationDataController.h
+++ b/Engine/Source/Runtime/Engine/Classes/Animation/AnimData/IAnimationDataController.h
@@ -5,6 +5,7 @@
 #include "CoreTypes.h"
 #include "UObject/Interface.h"
 #include "CurveIdentifier.h"
+#include "Misc/FrameRate.h"
 #include "AttributeIdentifier.h"
 #include "Algo/Transform.h"
 
diff --git a/Engine/Source/Runtime/Engine/Classes/Animation/AnimSequenceHelpers.h b/Engine/Source/Runtime/Engine/Classes/Animation/AnimSequenceHelpers.h
--- a/Engine/Source/Runtime/Engine/Classes/Animation/AnimSequenceHelpers.h
+++ b/Engine/Source/Runtime/Engine/Classes/Animation/AnimSequenceHelpers.h
@@ -10,6 +10,7 @@
 class UAnimSequence;
 class UAnimDataModel;
 class UAnimDataController;
+class IAnimationDataController;
 class FName;
 struct FCompactPose;
 struct FRawCurveTracks;
diff --git a/Engine/Source/Runtime/Engine/Private/Animation/AnimSequenceBase.cpp b/Engine/Source/Runtime/Engine/Private/Animation/AnimSequenceBase.cpp
--- a/Engine/Source/Runtime/Engine/Private/Animation/AnimSequenceBase.cpp
+++ b/Engine/Source/Runtime/Engine/Private/Animation/AnimSequenceBase.cpp
@@ -17,6 +17,7 @@
 
 #if WITH_EDITOR
 #include "IAnimationDataControllerModule.h"
+#include "Modules/ModuleManager.h"
 #endif // WITH_EDITOR
 
 #include "UObject/UE5MainStreamObjectVersion.h"
diff --git a/Engine/Source/Runtime/Engine/Private/Animation/AnimSequenceHelpers.cpp b/Engine/Source/Runtime/Engine/Private/Animation/AnimSequenceHelpers.cpp
--- a/Engine/Source/Runtime/Engine/Private/Animation/AnimSequenceHelpers.cpp
+++ b/Engine/Source/Runtime/Engine/Private/Animation/AnimSequenceHelpers.cpp
@@ -2,6 +2,7 @@
 
 #include "Animation/AnimSequenceHelpers.h"
 #include "Animation/AnimData/AnimDataModel.h"
+#include "Animation/AnimData/IAnimationDataController.h"
 
 #include "UObject/NameTypes.h"
 
