Animation moved from a different project and skeleton re-assigned in editor goes back to ref pose after save
#fix ensure when updating/removing tracks with invalid bone tree index also broadcast notify for track changed when bone index is updated in addition to removing bone tracks. This will force a resample of the source data - ensuring that bone indices are valid when compressing

diff --git a/Engine/Source/Developer/AnimationDataController/Private/AnimDataController.cpp b/Engine/Source/Developer/AnimationDataController/Private/AnimDataController.cpp
--- a/Engine/Source/Developer/AnimationDataController/Private/AnimDataController.cpp
+++ b/Engine/Source/Developer/AnimationDataController/Private/AnimDataController.cpp
@@ -337,6 +337,7 @@ bool UAnimDataController::RemoveBoneTracksMissingFromSkeleton(const USkeleton* S
 	if (Skeleton)
 	{
 		TArray<FName> TracksToBeRemoved;
+		TArray<FName> TracksUpdated;
 		const FReferenceSkeleton& ReferenceSkeleton = Skeleton->GetReferenceSkeleton();
 
 		for (FBoneAnimationTrack& Track : Model->BoneAnimationTracks)
@@ -359,6 +360,7 @@ bool UAnimDataController::RemoveBoneTracksMissingFromSkeleton(const USkeleton* S
 				{
 					// Update bone index
 					Track.BoneTreeIndex = BoneIndex;
+					TracksUpdated.Add(Track.Name);
 				}
 				else
 				{
@@ -369,16 +371,23 @@ bool UAnimDataController::RemoveBoneTracksMissingFromSkeleton(const USkeleton* S
 			}
 		}
 
-		if (TracksToBeRemoved.Num())
+		if (TracksToBeRemoved.Num() || TracksUpdated.Num())
 		{
 			CONDITIONAL_BRACKET(LOCTEXT("RemoveBoneTracksMissingFromSkeleton", "Validating Bone Animation Track Data against Skeleton"));
 			for (const FName& TrackName : TracksToBeRemoved)
 			{
 				RemoveBoneTrack(TrackName);
 			}
+
+			for (const FName& TrackName : TracksUpdated)
+			{
+				FAnimationTrackChangedPayload Payload;
+				Payload.Name = TrackName;
+				Model->Notify(EAnimDataModelNotifyType::TrackChanged, Payload);
+			}
 		}
 
-		return TracksToBeRemoved.Num() > 0;
+		return TracksToBeRemoved.Num() > 0 || TracksUpdated.Num() > 0;
 	}
 	else
 	{
