Remove behaviour of updating bone names, whenever animation tracks are validated against its Skeleton, as it has to ability to cause invalid data

diff --git a/Engine/Source/Developer/AnimationDataController/Private/AnimDataController.cpp b/Engine/Source/Developer/AnimationDataController/Private/AnimDataController.cpp
--- a/Engine/Source/Developer/AnimationDataController/Private/AnimDataController.cpp
+++ b/Engine/Source/Developer/AnimationDataController/Private/AnimDataController.cpp
@@ -342,33 +342,21 @@ bool UAnimDataController::RemoveBoneTracksMissingFromSkeleton(const USkeleton* S
 
 		for (FBoneAnimationTrack& Track : Model->BoneAnimationTracks)
 		{
-			if (ReferenceSkeleton.IsValidIndex(Track.BoneTreeIndex))
-			{
-				const FName BoneName = ReferenceSkeleton.GetBoneName(Track.BoneTreeIndex);
-				if (BoneName != Track.Name)
-				{
-					// Rename track			
-					Track.Name = BoneName;
-				}
-			}
-			else
-			{
-				// Try find correct bone index
-				const int32 BoneIndex = ReferenceSkeleton.FindBoneIndex(Track.Name);
+			// Try find correct bone index
+			const int32 BoneIndex = ReferenceSkeleton.FindBoneIndex(Track.Name);
 
-				if (BoneIndex != INDEX_NONE)
-				{
-					// Update bone index
-					Track.BoneTreeIndex = BoneIndex;
-					TracksUpdated.Add(Track.Name);
-				}
-				else
-				{
-					// Remove track
-					TracksToBeRemoved.Add(Track.Name);
-					ReportWarningf(LOCTEXT("InvalidBoneIndexWarning", "Unable to find bone index, animation track will be removed: {0}"), FText::FromName(Track.Name));
-				}
+			if (BoneIndex != INDEX_NONE && BoneIndex != Track.BoneTreeIndex)
+			{
+				// Update bone index
+				Track.BoneTreeIndex = BoneIndex;
+				TracksUpdated.Add(Track.Name);
 			}
+			else if (BoneIndex == INDEX_NONE)
+			{				
+				// Remove track
+				TracksToBeRemoved.Add(Track.Name);
+				ReportWarningf(LOCTEXT("InvalidBoneIndexWarning", "Unable to find bone index, animation track will be removed: {0}"), FText::FromName(Track.Name));				
+			}			
 		}
 
 		if (TracksToBeRemoved.Num() || TracksUpdated.Num())
