Add missing AnimDataController API for setting individual bone animation keys

PR #8920: [UE5] UAnimSequence change to allow partial edits (Contributed by [REDACTED_PERSON])

diff --git a/Engine/Source/Developer/AnimationDataController/Private/AnimDataController.cpp b/Engine/Source/Developer/AnimationDataController/Private/AnimDataController.cpp
--- a/Engine/Source/Developer/AnimationDataController/Private/AnimDataController.cpp
+++ b/Engine/Source/Developer/AnimationDataController/Private/AnimDataController.cpp
@@ -2,6 +2,7 @@
 
 #include "AnimDataController.h"
 #include "AnimDataControllerActions.h"
+#include "MovieSceneTimeHelpers.h"
 
 #include "Animation/AnimData/AnimDataModel.h"
 #include "Animation/AnimData/CurveIdentifier.h"
@@ -1557,6 +1558,193 @@ bool UAnimDataController::SetBoneTrackKeys(FName BoneName, const TArray<FVector3
 	return false;
 }
 
+int32 DiscreteInclusiveLower(const TRange<int32>& InRange)
+{
+	check(!InRange.GetLowerBound().IsOpen());
+
+	// Add one for exclusive lower bounds since they start on the next subsequent frame
+	static const int32 Offsets[]   = { 0, 1 };
+	const int32        OffsetIndex = (int32)InRange.GetLowerBound().IsExclusive();
+
+	return InRange.GetLowerBound().GetValue() + Offsets[OffsetIndex];
+}
+
+int32 DiscreteExclusiveUpper(const TRange<int32>& InRange)
+{
+	check(!InRange.GetUpperBound().IsOpen());
+
+	// Add one for inclusive upper bounds since they finish on the next subsequent frame
+	static const int32 Offsets[]   = { 0, 1 };
+	const int32        OffsetIndex = (int32)InRange.GetUpperBound().IsInclusive();
+
+	return InRange.GetUpperBound().GetValue() + Offsets[OffsetIndex];
+}
+
+bool UAnimDataController::UpdateBoneTrackKeys(FName BoneName, const FInt32Range& KeyRangeToSet, const TArray<FVector3f>& PositionalKeys, const TArray<FQuat4f>& RotationalKeys, const TArray<FVector3f>& ScalingKeys, bool bShouldTransact)
+{	
+	if (!CheckOuterClass(UAnimSequence::StaticClass()))
+	{
+		return false;
+	}
+
+	CONDITIONAL_TRANSACTION(LOCTEXT("SetTrackKeysRangeTransaction", "Setting Animation Data Track keys"));
+
+	// Validate key format
+	const int32 MaxNumKeys = FMath::Max(FMath::Max(PositionalKeys.Num(), RotationalKeys.Num()), ScalingKeys.Num());
+	const int32 RangeMin = DiscreteInclusiveLower(KeyRangeToSet);
+	const int32 RangeMax = DiscreteExclusiveUpper(KeyRangeToSet);
+	if (MaxNumKeys > 0)
+	{
+		const bool bValidPosKeys = PositionalKeys.Num() == MaxNumKeys;
+		const bool bValidRotKeys = RotationalKeys.Num() == MaxNumKeys;
+		const bool bValidScaleKeys = ScalingKeys.Num() == MaxNumKeys;
+
+		if (bValidPosKeys && bValidRotKeys && bValidScaleKeys)
+		{
+			const int32 NumKeysToSet = RangeMax - RangeMin;
+			if (NumKeysToSet == MaxNumKeys)
+			{
+				if (FBoneAnimationTrack* TrackPtr = Model->FindMutableBoneTrackByName(BoneName))
+				{
+					const FInt32Range TrackKeyRange(0, TrackPtr->InternalTrackData.PosKeys.Num());
+					if(TrackKeyRange.Contains(KeyRangeToSet))
+					{
+						FRawAnimSequenceTrack& InternalTrackData = TrackPtr->InternalTrackData;
+						TArray<FVector3f>& TrackPosKeys = InternalTrackData.PosKeys;
+						TArray<FQuat4f>& TrackRotKeys = InternalTrackData.RotKeys;
+						TArray<FVector3f>& TrackScaleKeys = InternalTrackData.ScaleKeys;
+
+						CONDITIONAL_ACTION(UE::Anim::FSetTrackKeysAction, *TrackPtr);
+
+						int32 KeyIndex = 0;
+						for (int32 FrameIndex = RangeMin; FrameIndex < RangeMax; ++FrameIndex, ++KeyIndex)
+						{
+							TrackPosKeys[FrameIndex] = PositionalKeys[KeyIndex];
+							TrackRotKeys[FrameIndex] = RotationalKeys[KeyIndex];
+							TrackScaleKeys[FrameIndex] = ScalingKeys[KeyIndex];
+						}
+
+						FAnimationTrackChangedPayload Payload;
+						Payload.Name = BoneName;
+
+						Model->Notify(EAnimDataModelNotifyType::TrackChanged, Payload);
+
+						return true;
+					}
+					else
+					{
+						ReportWarningf(LOCTEXT("InvalidTrackSetFrameRange", "Range of to-be-set bone track (with name {0}) keys does not overlap existing range"), FText::FromName(BoneName));
+					}
+				}
+				else
+				{
+					ReportWarningf(LOCTEXT("InvalidTrackNameWarning", "Track with name {0} does not exist"), FText::FromName(BoneName));
+				}
+			}
+			else
+			{
+				ReportErrorf(LOCTEXT("InvalidKeyIndexRangeError",
+								 "Invalid key index range bound [{0}, {1}], expected to equal the size of the positional, rotational and scaling keys {2}"),
+								 FText::AsNumber(RangeMin),
+								 FText::AsNumber(RangeMax),
+								 FText::AsNumber(MaxNumKeys));
+			}	
+		}
+		else
+		{
+			ReportErrorf(LOCTEXT("InvalidTrackKeyDataError", "Invalid track key data, expected uniform data: number of positional keys {0}, number of rotational keys {1}, number of scaling keys {2}"), FText::AsNumber(PositionalKeys.Num()), FText::AsNumber(RotationalKeys.Num()), FText::AsNumber(ScalingKeys.Num()));
+		}
+	}
+	else
+	{
+		ReportErrorf(LOCTEXT("MissingTrackKeyDataError", "Missing track key data, expected uniform data: number of positional keys {0}, number of rotational keys {1}, number of scaling keys {2}"), FText::AsNumber(PositionalKeys.Num()), FText::AsNumber(RotationalKeys.Num()), FText::AsNumber(ScalingKeys.Num()));
+	}
+
+	return false;
+}
+
+bool UAnimDataController::UpdateBoneTrackKeys(FName BoneName, const FInt32Range& KeyRangeToSet, const TArray<FVector>& PositionalKeys, const TArray<FQuat>& RotationalKeys, const TArray<FVector>& ScalingKeys, bool bShouldTransact)
+{
+	if (!CheckOuterClass(UAnimSequence::StaticClass()))
+	{
+		return false;
+	}
+
+	CONDITIONAL_TRANSACTION(LOCTEXT("SetTrackKeysTransaction", "Setting Animation Data Track keys"));
+
+	// Validate key format
+	const int32 MaxNumKeys = FMath::Max(FMath::Max(PositionalKeys.Num(), RotationalKeys.Num()), ScalingKeys.Num());
+	const int32 RangeMin = DiscreteInclusiveLower(KeyRangeToSet);
+	const int32 RangeMax = DiscreteExclusiveUpper(KeyRangeToSet);
+	if (MaxNumKeys > 0)
+	{
+		const bool bValidPosKeys = PositionalKeys.Num() == MaxNumKeys;
+		const bool bValidRotKeys = RotationalKeys.Num() == MaxNumKeys;
+		const bool bValidScaleKeys = ScalingKeys.Num() == MaxNumKeys;
+
+		if (bValidPosKeys && bValidRotKeys && bValidScaleKeys)
+		{
+			const int32 NumKeysToSet = RangeMax - RangeMin;
+			if (NumKeysToSet == MaxNumKeys)
+			{
+				if (FBoneAnimationTrack* TrackPtr = Model->FindMutableBoneTrackByName(BoneName))
+				{
+					const FInt32Range TrackKeyRange(0, TrackPtr->InternalTrackData.PosKeys.Num());
+					if(TrackKeyRange.Contains(KeyRangeToSet))
+					{
+						FRawAnimSequenceTrack& InternalTrackData = TrackPtr->InternalTrackData;
+						TArray<FVector3f>& TrackPosKeys = InternalTrackData.PosKeys;
+						TArray<FQuat4f>& TrackRotKeys = InternalTrackData.RotKeys;
+						TArray<FVector3f>& TrackScaleKeys = InternalTrackData.ScaleKeys;
+
+						CONDITIONAL_ACTION(UE::Anim::FSetTrackKeysAction, *TrackPtr);
+
+						int32 KeyIndex = 0;
+						for (int32 FrameIndex = RangeMin; FrameIndex < RangeMax; ++FrameIndex, ++KeyIndex)
+						{
+							TrackPosKeys[FrameIndex] = FVector3f(PositionalKeys[KeyIndex]);
+							TrackRotKeys[FrameIndex] = FQuat4f(RotationalKeys[KeyIndex]);
+							TrackScaleKeys[FrameIndex] = FVector3f(ScalingKeys[KeyIndex]);
+						}
+
+						FAnimationTrackChangedPayload Payload;
+						Payload.Name = BoneName;
+
+						Model->Notify(EAnimDataModelNotifyType::TrackChanged, Payload);
+
+						return true;
+					}
+					else
+					{
+						ReportWarningf(LOCTEXT("InvalidTrackSetFrameRange", "Range of to-be-set bone track (with name {0}) keys does not overlap existing range"), FText::FromName(BoneName));
+					}
+				}
+				else
+				{
+					ReportWarningf(LOCTEXT("InvalidTrackNameWarning", "Track with name {0} does not exist"), FText::FromName(BoneName));
+				}
+			}
+			else
+			{
+				ReportErrorf(LOCTEXT("InvalidKeyIndexRangeError",
+								 "Invalid key index range bound [{0}, {1}], expected to equal the size of the positional, rotational and scaling keys {2}"),
+								 FText::AsNumber(KeyRangeToSet.GetLowerBoundValue()),
+								 FText::AsNumber(KeyRangeToSet.GetUpperBoundValue()),
+								 FText::AsNumber(MaxNumKeys));
+			}	
+		}
+		else
+		{
+			ReportErrorf(LOCTEXT("InvalidTrackKeyDataError", "Invalid track key data, expected uniform data: number of positional keys {0}, number of rotational keys {1}, number of scaling keys {2}"), FText::AsNumber(PositionalKeys.Num()), FText::AsNumber(RotationalKeys.Num()), FText::AsNumber(ScalingKeys.Num()));
+		}
+	}
+	else
+	{
+		ReportErrorf(LOCTEXT("MissingTrackKeyDataError", "Missing track key data, expected uniform data: number of positional keys {0}, number of rotational keys {1}, number of scaling keys {2}"), FText::AsNumber(PositionalKeys.Num()), FText::AsNumber(RotationalKeys.Num()), FText::AsNumber(ScalingKeys.Num()));
+	}
+
+	return false;
+}
 
 void UAnimDataController::ResizeCurves(float NewLength, bool bInserted, float T0, float T1, bool bShouldTransact /*= true*/)
 {
diff --git a/Engine/Source/Developer/AnimationDataController/Public/AnimDataController.h b/Engine/Source/Developer/AnimationDataController/Public/AnimDataController.h
--- a/Engine/Source/Developer/AnimationDataController/Public/AnimDataController.h
+++ b/Engine/Source/Developer/AnimationDataController/Public/AnimDataController.h
@@ -59,6 +59,8 @@ public:
 	virtual void RemoveAllBoneTracks(bool bShouldTransact = true) override;
 	virtual bool SetBoneTrackKeys(FName BoneName, const TArray<FVector3f>& PositionalKeys, const TArray<FQuat4f>& RotationalKeys, const TArray<FVector3f>& ScalingKeys, bool bShouldTransact = true) override;	
 	virtual bool SetBoneTrackKeys(FName BoneName, const TArray<FVector>& PositionalKeys, const TArray<FQuat>& RotationalKeys, const TArray<FVector>& ScalingKeys, bool bShouldTransact = true) override;	
+	virtual bool UpdateBoneTrackKeys(FName BoneName, const FInt32Range& KeyRangeToSet, const TArray<FVector3f>& PositionalKeys, const TArray<FQuat4f>& RotationalKeys, const TArray<FVector3f>& ScalingKeys, bool bShouldTransact = true) override;
+	virtual bool UpdateBoneTrackKeys(FName BoneName, const FInt32Range& KeyRangeToSet, const TArray<FVector>& PositionalKeys, const TArray<FQuat>& RotationalKeys, const TArray<FVector>& ScalingKeys, bool bShouldTransact = true) override;
 	virtual bool AddCurve(const FAnimationCurveIdentifier& CurveId, int32 CurveFlags = 0x00000004, bool bShouldTransact = true) override;
 	virtual bool DuplicateCurve(const FAnimationCurveIdentifier& CopyCurveId, const FAnimationCurveIdentifier& NewCurveId, bool bShouldTransact = true) override;
 	virtual bool RemoveCurve(const FAnimationCurveIdentifier& CurveId, bool bShouldTransact = true) override;
diff --git a/Engine/Source/Runtime/Engine/Classes/Animation/AnimData/IAnimationDataController.h b/Engine/Source/Runtime/Engine/Classes/Animation/AnimData/IAnimationDataController.h
--- a/Engine/Source/Runtime/Engine/Classes/Animation/AnimData/IAnimationDataController.h
+++ b/Engine/Source/Runtime/Engine/Classes/Animation/AnimData/IAnimationDataController.h
@@ -224,7 +224,37 @@ public:
 	*/
 	UFUNCTION(BlueprintCallable, Category = AnimationData)
 	virtual bool SetBoneTrackKeys(FName BoneName, const TArray<FVector>& PositionalKeys, const TArray<FQuat>& RotationalKeys, const TArray<FVector>& ScalingKeys, bool bShouldTransact = true) = 0;
-	
+
+	/**
+	* Sets a range of keys for an existing bone animation track with the provided name. Broadcasts a EAnimDataModelNotifyType::TrackChanged notify if successful.
+	* The provided number of keys provided is expected to match for each component, be between FrameLowerBound and FrameUpperBound (inclusive), and be non-zero.
+	*
+	* @param	BoneName			Bone name of the track for which the keys should be set
+	* @param	FrameRange			Range of frames to set keys for
+	* @param	PositionalKeys		Array of keys for the translation component
+	* @param	RotationalKeys		Array of keys for the rotation component
+	* @param	ScalingKeys			Array of keys for the scale component
+	* @param	bShouldTransact		Whether or not any undo-redo changes should be generated
+	*
+	* @return	Whether or not the keys were successfully set
+	*/
+	virtual bool UpdateBoneTrackKeys(FName BoneName, const FInt32Range& KeyRangeToSet, const TArray<FVector3f>& PositionalKeys, const TArray<FQuat4f>& RotationalKeys, const TArray<FVector3f>& ScalingKeys, bool bShouldTransact = true) = 0;
+
+	/**
+	* Sets a range of keys for an existing bone animation track with the provided name. Broadcasts a EAnimDataModelNotifyType::TrackChanged notify if successful.
+	* The provided number of keys provided is expected to match for each component, be between FrameLowerBound and FrameUpperBound (inclusive), and be non-zero.
+	*
+	* @param	BoneName			Bone name of the track for which the keys should be set
+	* @param	FrameRange			Range of frames to set keys for
+	* @param	PositionalKeys		Array of keys for the translation component
+	* @param	RotationalKeys		Array of keys for the rotation component
+	* @param	ScalingKeys			Array of keys for the scale component
+	* @param	bShouldTransact		Whether or not any undo-redo changes should be generated
+	*
+	* @return	Whether or not the keys were successfully set
+	*/
+	virtual bool UpdateBoneTrackKeys(FName BoneName, const FInt32Range& KeyRangeToSet, const TArray<FVector>& PositionalKeys, const TArray<FQuat>& RotationalKeys, const TArray<FVector>& ScalingKeys, bool bShouldTransact = true) = 0;
+
 	/**
 	* Adds a new curve with the provided information. Broadcasts a EAnimDataModelNotifyType::CurveAdded notify if successful.
 	*
diff --git a/Engine/Source/Runtime/Engine/Classes/Animation/AnimSequence.h b/Engine/Source/Runtime/Engine/Classes/Animation/AnimSequence.h
--- a/Engine/Source/Runtime/Engine/Classes/Animation/AnimSequence.h
+++ b/Engine/Source/Runtime/Engine/Classes/Animation/AnimSequence.h
@@ -36,6 +36,7 @@ struct FAnimCompressContext;
 struct FAnimSequenceDecompressionContext;
 struct FCompactPose;
 
+namespace UE { namespace Anim { namespace Compression { struct FScopedCompressionGuard; } } }
 
 // These two always should go together, but it is not right now. 
 // I wonder in the future, we change all compressed to be inside as well, so they all stay together
@@ -1110,6 +1111,6 @@ public:
 	friend class UAnimationBlueprintLibrary;
 	friend class UAnimBoneCompressionSettings;
 	friend class FCustomAttributeCustomization;
-	friend class UAnimationBlueprintLibrary;
 	friend class FAnimSequenceTestBase;
+	friend struct UE::Anim::Compression::FScopedCompressionGuard;
 };
diff --git a/Engine/Source/Runtime/Engine/Classes/Animation/AnimSequenceHelpers.h b/Engine/Source/Runtime/Engine/Classes/Animation/AnimSequenceHelpers.h
--- a/Engine/Source/Runtime/Engine/Classes/Animation/AnimSequenceHelpers.h
+++ b/Engine/Source/Runtime/Engine/Classes/Animation/AnimSequenceHelpers.h
@@ -198,6 +198,18 @@ namespace Anim {
 		* @param	RawTrack			Individual track to sanitize
 		*/
 		ENGINE_API void SanitizeRawAnimSequenceTrack(FRawAnimSequenceTrack& RawTrack);
+
+#if WITH_EDITOR
+		// RAII helper to temporarily block Animation compression requests for specified AnimationSequence
+		struct ENGINE_API FScopedCompressionGuard
+		{
+			FScopedCompressionGuard() = delete;
+			FScopedCompressionGuard(UAnimSequence* InAnimSequence);
+			~FScopedCompressionGuard();
+		protected:
+			TObjectPtr<UAnimSequence> AnimSequence;
+		};
+#endif // WITH_EDITOR
 	}
 
 	namespace Retargeting
diff --git a/Engine/Source/Runtime/Engine/Private/Animation/AnimSequenceHelpers.cpp b/Engine/Source/Runtime/Engine/Private/Animation/AnimSequenceHelpers.cpp
--- a/Engine/Source/Runtime/Engine/Private/Animation/AnimSequenceHelpers.cpp
+++ b/Engine/Source/Runtime/Engine/Private/Animation/AnimSequenceHelpers.cpp
@@ -716,6 +716,22 @@ void Compression::SanitizeRawAnimSequenceTrack(FRawAnimSequenceTrack& RawTrack)
 }
 
 #if WITH_EDITOR
+Compression::FScopedCompressionGuard::FScopedCompressionGuard(UAnimSequence* InAnimSequence) : AnimSequence(InAnimSequence)
+{
+	if(AnimSequence)
+	{
+		AnimSequence->bBlockCompressionRequests = true;		
+	}
+}
+
+Compression::FScopedCompressionGuard::~FScopedCompressionGuard()
+{
+	if(AnimSequence)
+	{
+		AnimSequence->bBlockCompressionRequests = false;		
+	}
+}
+
 bool AnimationData::AddLoopingInterpolation(UAnimSequence* InSequence)
 {
 	const UAnimDataModel* DataModel = InSequence->GetDataModel();
