**New**

diff --git a/Engine/Source/Developer/AnimationDataController/Private/AnimDataController.cpp b/Engine/Source/Developer/AnimationDataController/Private/AnimDataController.cpp
--- a/Engine/Source/Developer/AnimationDataController/Private/AnimDataController.cpp
+++ b/Engine/Source/Developer/AnimationDataController/Private/AnimDataController.cpp
@@ -4,89 +4,42 @@
 #include "AnimDataControllerActions.h"
 #include "MovieSceneTimeHelpers.h"
 
-#include "Animation/AnimData/AnimDataModel.h"
+#include "Animation/AnimData/IAnimationDataModel.h"
 #include "Animation/AnimData/CurveIdentifier.h"
-
-#include "Animation/AnimSequenceBase.h"
 #include "Animation/AnimSequence.h"
 
-
 #include "Algo/Transform.h"
+#include "Animation/AnimationSettings.h"
 #include "UObject/NameTypes.h"
 #include "Animation/AnimCurveTypes.h"
+#include "Animation/AnimSequenceHelpers.h"
 #include "Math/UnrealMathUtility.h"
 
 #define LOCTEXT_NAMESPACE "AnimDataController"
 
 #if WITH_EDITOR
-
-namespace UE {
-namespace Anim {
-	bool CanTransactChanges()
-	{
-		return GEngine && GEngine->CanTransact() && !GIsTransacting;
-	}
-
-	struct FScopedCompoundTransaction
-	{
-		FScopedCompoundTransaction(UE::FChangeTransactor& InTransactor, const FText& InDescription) : Transactor(InTransactor), bCreated(false)
-		{
-			if (CanTransactChanges() && !Transactor.IsTransactionPending())
-			{
-				Transactor.OpenTransaction(InDescription);
-				bCreated = true;
-			}
-		}
-
-		~FScopedCompoundTransaction()
-		{
-			if (bCreated)
-			{
-				Transactor.CloseTransaction();
-			}
-		}
-
-		UE::FChangeTransactor& Transactor;
-		bool bCreated;
-	};
-}}
-
-#define CONDITIONAL_TRANSACTION(Text) \
-	TUniquePtr<UE::Anim::FScopedCompoundTransaction> Transaction; \
-	if (UE::Anim::CanTransactChanges() && bShouldTransact) \
-	{ \
-		Transaction = MakeUnique<UE::Anim::FScopedCompoundTransaction>(ChangeTransactor, Text); \
-	}
-
-#define CONDITIONAL_BRACKET(Text) IAnimationDataController::FScopedBracket Transaction(this, Text, UE::Anim::CanTransactChanges() && bShouldTransact);
-
-#define CONDITIONAL_ACTION(ActionClass, ...) \
-	if (UE::Anim::CanTransactChanges() && bShouldTransact) \
-	{ \
-		ChangeTransactor.AddTransactionChange<ActionClass>(__VA_ARGS__); \
-	}
-
-void UAnimDataController::SetModel(UAnimDataModel* InModel)
+void UAnimDataController::SetModel(TScriptInterface<IAnimationDataModel> InModel)
 {	
 	if (Model != nullptr)
 	{
 		Model->GetModifiedEvent().RemoveAll(this);
 	}
 
-	Model = InModel;
+	ModelInterface = InModel;
+	Model = CastChecked<UAnimDataModel>(InModel.GetObject(), ECastCheckedType::NullAllowed);
 	
-	ChangeTransactor.SetTransactionObject(InModel);
+	ChangeTransactor.SetTransactionObject(Model);
 }
 
 void UAnimDataController::OpenBracket(const FText& InTitle, bool bShouldTransact /*= true*/)
 {
 	ValidateModel();
 
-	if (UE::Anim::CanTransactChanges() && !ChangeTransactor.IsTransactionPending())
+	if (UE::FChangeTransactor::CanTransactChanges() && !ChangeTransactor.IsTransactionPending())
 	{
 		ChangeTransactor.OpenTransaction(InTitle);
 
-		CONDITIONAL_ACTION(UE::Anim::FCloseBracketAction, InTitle.ToString());
+		ConditionalAction<UE::Anim::FCloseBracketAction>(bShouldTransact, InTitle.ToString());
 	}
 
 	if (BracketDepth == 0)
@@ -94,7 +47,7 @@ void UAnimDataController::OpenBracket(const FText& InTitle, bool bShouldTransact
 		FBracketPayload Payload;
 		Payload.Description = InTitle.ToString();
 
-		Model->Notify(EAnimDataModelNotifyType::BracketOpened, Payload);
+		Model->GetNotifier().Notify(EAnimDataModelNotifyType::BracketOpened, Payload);
 	}
 
 	++BracketDepth;
@@ -114,95 +67,181 @@ void UAnimDataController::CloseBracket(bool bShouldTransact /*= true*/)
 
 	if (BracketDepth == 0)
 	{
-		if (UE::Anim::CanTransactChanges())
+		if (UE::FChangeTransactor::CanTransactChanges())
 		{
 			ensure(ChangeTransactor.IsTransactionPending());
 
-			CONDITIONAL_ACTION(UE::Anim::FOpenBracketAction, TEXT("Open Bracket"));
+			ConditionalAction<UE::Anim::FOpenBracketAction>(bShouldTransact, TEXT("Open Bracket"));
 
 			ChangeTransactor.CloseTransaction();
 		}
 		
-		Model->Notify(EAnimDataModelNotifyType::BracketClosed);
+		Model->GetNotifier().Notify(EAnimDataModelNotifyType::BracketClosed);
 	}
 }
 
-void UAnimDataController::SetPlayLength(float Length, bool bShouldTransact /*= true*/)
+void UAnimDataController::SetNumberOfFrames(FFrameNumber Length, bool bShouldTransact)
 {
 	ValidateModel();
+	const FFrameNumber CurrentNumberOfFrames = Model->GetNumberOfFrames();
 
-	// Calculate whether or new play length is shorter or longer than current, set-up T0; T1 accordingly
-	// Assumption is made that time is always added or removed at/from the end
-	// Added: T0 = current length, T1 = new length
-	// Removed: T0 = current length - removed length, T1 = current length
-	const float Delta = FMath::Abs(Length - Model->PlayLength);
-	const float T0 = Length > Model->PlayLength ? Model->PlayLength : Model->PlayLength - Delta;
-	const float T1 = Length > Model->PlayLength ? Length : Model->PlayLength;
-	ResizePlayLength(Length, T0, T1, bShouldTransact);	
+	const int32 DeltaFrames = FMath::Abs(Length.Value - CurrentNumberOfFrames.Value);
+	
+	const FFrameNumber T0 = Length > CurrentNumberOfFrames ? CurrentNumberOfFrames : CurrentNumberOfFrames - DeltaFrames;
+	const FFrameNumber T1 = Length > CurrentNumberOfFrames ? Length : CurrentNumberOfFrames;
+
+	ResizeNumberOfFrames(Length, T0, T1, bShouldTransact);
 }
 
-void UAnimDataController::Resize(float Length, float T0, float T1, bool bShouldTransact /*= true*/)
+void UAnimDataController::ResizeNumberOfFrames(FFrameNumber NewLength, FFrameNumber T0, FFrameNumber T1, bool bShouldTransact)
 {
 	ValidateModel();
 	
-	const TRange<float> PlayRange(TRange<float>::BoundsType::Inclusive(0.f), TRange<float>::BoundsType::Inclusive(Model->PlayLength));
-	if (!FMath::IsNearlyZero(Length) && Length > 0.f)
+	const TRange<FFrameNumber> PlayRange(TRange<FFrameNumber>::BoundsType::Inclusive(0), TRange<FFrameNumber>::BoundsType::Exclusive(FMath::Max(1, Model->GetNumberOfKeys())));
+	if (NewLength >= 0)
 	{
-		if (Length != Model->PlayLength)
+		if (NewLength != Model->GetNumberOfFrames())
 		{
-			// Ensure that T0 is within the curent play range
+			// Ensure that T0 is within the current play range
 			if (PlayRange.Contains(T0))
 			{
 				// Ensure that the start and end length of either removal or insertion are valid
 				if (T0 < T1)
 				{
-					CONDITIONAL_BRACKET(LOCTEXT("ResizeModel", "Resizing Animation Data"));
-					const bool bInserted = Length > Model->PlayLength;
-					ResizePlayLength(Length, T0, T1, bShouldTransact);
-					ResizeCurves(Length, bInserted, T0, T1, bShouldTransact);
-					ResizeAttributes(Length, bInserted, T0, T1, bShouldTransact);
+					FTransaction Transaction = ConditionalTransaction(LOCTEXT("ResizePlayLength", "Resizing Play Length"), bShouldTransact);
+
+					const FFrameRate CurrentFrameRate = Model->GetFrameRate();
+					const FFrameNumber CurrentNumberOfFrames = Model->GetNumberOfFrames();
+
+					FSequenceLengthChangedPayload Payload;
+					PRAGMA_DISABLE_DEPRECATION_WARNINGS
+					Payload.T0 = CurrentFrameRate.AsSeconds(T0);
+					Payload.T1 = CurrentFrameRate.AsSeconds(T1);
+					Payload.PreviousLength = CurrentFrameRate.AsSeconds(CurrentNumberOfFrames);
+					PRAGMA_ENABLE_DEPRECATION_WARNINGS
+
+					Payload.PreviousNumberOfFrames = CurrentNumberOfFrames;
+					Payload.Frame0 = T0;
+					Payload.Frame1 = T1;
+
+					ConditionalAction<UE::Anim::FResizePlayLengthInFramesAction>(bShouldTransact, Model, Payload.Frame0, Payload.Frame1);
+
+					Model->NumberOfFrames = NewLength.Value;
+					Model->NumberOfKeys = Model->NumberOfFrames + 1;
+	
+					Model->GetNotifier().Notify<FSequenceLengthChangedPayload>(EAnimDataModelNotifyType::SequenceLengthChanged, Payload);
 				}
 				else
 				{
-					ReportErrorf(LOCTEXT("InvalidEndTimeError", "Invalid T1, smaller that T0 value: T0 {0}, T1 {1}"), FText::AsNumber(T0), FText::AsNumber(T1));
+					ReportErrorf(LOCTEXT("InvalidEndTimeError", "Invalid T1, smaller that T0 value: T0 {0}, T1 {1}"), FText::AsNumber(T0.Value), FText::AsNumber(T0.Value));
 				}
 			}
 			else
 			{
-				ReportErrorf(LOCTEXT("InvalidStartTimeError", "Invalid T0, not within existing play range: T0 {0}, Play Length {1}"), FText::AsNumber(T0), FText::AsNumber(Model->PlayLength));
-			}			
+				ReportErrorf(LOCTEXT("InvalidStartTimeError", "Invalid T0, not within existing play range: T0 {0}, Play Length {1}"), FText::AsNumber(T0.Value), FText::AsNumber(Model->GetPlayLength()));
+			}
 		}
-		else
+		else if (Model->bPopulated)
 		{
-			ReportWarningf(LOCTEXT("SamePlayLengthWarning", "New play length is same as existing one: {0} seconds"), FText::AsNumber(Length));
+			ReportWarningf(LOCTEXT("SamePlayLengthWarning", "New play length is same as existing one: {0} frames"), FText::AsNumber(NewLength.Value));
 		}
 	}
 	else
 	{
-		ReportErrorf(LOCTEXT("InvalidPlayLengthError", "Invalid play length value provided: {0} seconds"), FText::AsNumber(Length));
+		ReportErrorf(LOCTEXT("InvalidPlayLengthError", "Invalid play length value provided: {0} frames"), FText::AsNumber(NewLength.Value));
 	}
 }
 
-void UAnimDataController::SetFrameRate(FFrameRate FrameRate, bool bShouldTransact /*= true*/)
+void UAnimDataController::ResizeInFrames(FFrameNumber NewLength, FFrameNumber T0, FFrameNumber T1, bool bShouldTransact)
 {
 	ValidateModel();
-
-	// Disallow invalid frame-rates, or 0.0 intervals
-	const float FrameRateInterval = FrameRate.AsInterval();
-	if ( FrameRate.IsValid() && !FMath::IsNearlyZero(FrameRateInterval) && FrameRateInterval > 0.f)
+	
+	const int32 CurrentNumberOFrames = Model->GetNumberOfFrames();
+	
+	const TRange<FFrameNumber> PlayRange(TRange<FFrameNumber>::BoundsType::Inclusive(0), TRange<FFrameNumber>::BoundsType::Exclusive(FMath::Max(1,Model->GetNumberOfKeys())));
+	if (NewLength >= 0)
 	{
-		CONDITIONAL_TRANSACTION(LOCTEXT("SetFrameRate", "Setting Frame Rate"));
+		if (NewLength != Model->GetNumberOfFrames())
+		{
+			// Ensure that T0 is within the current play range
+			if (PlayRange.Contains(T0))
+			{
+				// Ensure that the start and end length of either removal or insertion are valid
+				if (T0 < T1)
+				{
+					FBracket Bracket = ConditionalBracket(LOCTEXT("ResizeModel", "Resizing Animation Data"), bShouldTransact);
 
-		CONDITIONAL_ACTION(UE::Anim::FSetFrameRateAction, Model);
+					const bool bInserted = NewLength > CurrentNumberOFrames;
+					ResizeNumberOfFrames(NewLength, T0, T1, bShouldTransact);
+					
+					const FFrameRate CurrentFrameRate = Model->GetFrameRate();
+					ResizeCurves(CurrentFrameRate.AsSeconds(NewLength), bInserted, CurrentFrameRate.AsSeconds(T0), CurrentFrameRate.AsSeconds(T1), bShouldTransact);
+					ResizeAttributes(CurrentFrameRate.AsSeconds(NewLength), bInserted, CurrentFrameRate.AsSeconds(T0), CurrentFrameRate.AsSeconds(T1), bShouldTransact);
+				}
+				else
+                {
+                	ReportErrorf(LOCTEXT("InvalidEndTimeError", "Invalid T1, smaller that T0 value: T0 {0}, T1 {1}"), FText::AsNumber(T0.Value), FText::AsNumber(T1.Value));
+                }
+            }
+            else
+            {
+                ReportErrorf(LOCTEXT("InvalidStartTimeError", "Invalid T0, not within existing play range: T0 {0}, Play Length {1}"), FText::AsNumber(T0.Value), FText::AsNumber(CurrentNumberOFrames));
+            }			
+        }
+		else if (Model->bPopulated)
+        {
+            ReportWarningf(LOCTEXT("SameGetPlayLengthWarning", "New play length is same as existing one: {0} frames"), FText::AsNumber(CurrentNumberOFrames));
+        }
+    }
+    else
+    {
+        ReportErrorf(LOCTEXT("InvalidGetPlayLengthError", "Invalid play length value provided: {0} frames"), FText::AsNumber(CurrentNumberOFrames));
+    }
+}
 
-		FFrameRateChangedPayload Payload;
-		Payload.PreviousFrameRate = Model->FrameRate;
+void UAnimDataController::SetPlayLength(float Length, bool bShouldTransact /*= true*/)
+{
+	SetNumberOfFrames(ConvertSecondsToFrameNumber(Length), bShouldTransact);
+}
 
-		Model->FrameRate = FrameRate;
-		Model->NumberOfFrames = Model->FrameRate.AsFrameTime(Model->PlayLength).RoundToFrame().Value;
-		Model->NumberOfKeys = Model->NumberOfFrames + 1;
+void UAnimDataController::Resize(float Length, float T0, float T1, bool bShouldTransact /*= true*/)
+{
+	ValidateModel();
+	
+	ResizeInFrames(ConvertSecondsToFrameNumber(Length), ConvertSecondsToFrameNumber(T0), ConvertSecondsToFrameNumber(T1), bShouldTransact);
+}
 
-		Model->Notify(EAnimDataModelNotifyType::FrameRateChanged, Payload);
+void UAnimDataController::SetFrameRate(FFrameRate FrameRate, bool bShouldTransact /*= true*/)
+{
+	ValidateModel();
+
+	// Disallow invalid frame-rates, or 0.0 intervals
+	const double FrameRateInterval = FrameRate.AsInterval();
+	if ( FrameRate.IsValid() && !FMath::IsNearlyZero(FrameRateInterval) && FrameRateInterval > 0.0)
+	{
+		// Need to verify framerate
+		const FFrameRate CurrentFrameRate = Model->GetFrameRate();
+		if (FrameRate.IsMultipleOf(CurrentFrameRate) || FrameRate.IsFactorOf(CurrentFrameRate) || !Model->bPopulated)
+		{
+			FTransaction Transaction = ConditionalTransaction(LOCTEXT("SetFrameRate", "Setting Frame Rate"), bShouldTransact);
+			ConditionalAction<UE::Anim::FSetFrameRateAction>(bShouldTransact, Model);
+
+			const FFrameNumber CurrentNumberOfFrames = Model->GetNumberOfFrames();
+			const FFrameTime ConvertedLastFrameTime = FFrameRate::TransformTime(CurrentNumberOfFrames, CurrentFrameRate, FrameRate);
+			ensure(FMath::IsNearlyZero(ConvertedLastFrameTime.GetSubFrame()) || !Model->bPopulated);
+			
+			Model->FrameRate = FrameRate;
+			Model->NumberOfFrames = ConvertedLastFrameTime.GetFrame().Value;
+			Model->NumberOfKeys = Model->NumberOfFrames + 1;
+			
+			FFrameRateChangedPayload Payload;
+			Payload.PreviousFrameRate = CurrentFrameRate;
+			Model->GetNotifier().Notify(EAnimDataModelNotifyType::FrameRateChanged, Payload);
+		}
+		else
+        {
+        	ReportErrorf(LOCTEXT("NonCompatibleFrameRateError", "Incompatible frame rate provided: {0} not a multiple or fact or {1}"), FrameRate.ToPrettyText(), CurrentFrameRate.ToPrettyText());
+        }
 	}
 	else
 	{
@@ -219,7 +258,7 @@ void UAnimDataController::UpdateCurveNamesFromSkeleton(const USkeleton* Skeleton
 	{
 		if (IsSupportedCurveType(SupportedCurveType))
 		{
-			CONDITIONAL_BRACKET(LOCTEXT("ValidateRawCurves", "Validating Animation Curve Names"));
+			FBracket Bracket = ConditionalBracket(LOCTEXT("ValidateRawCurves", "Validating Animation Curve Names"), bShouldTransact);
 			switch (SupportedCurveType)
 			{
 			case ERawCurveTrackTypes::RCT_Float:
@@ -254,12 +293,17 @@ void UAnimDataController::UpdateCurveNamesFromSkeleton(const USkeleton* Skeleton
 				}
 				break;
 			}
+			case ERawCurveTrackTypes::RCT_Vector:
+			default:
+				const FString CurveTypeAsString = GetCurveTypeValueName(SupportedCurveType);
+				ReportWarningf(LOCTEXT("InvalidCurveTypeWarning", "Invalid curve type provided: {0} ({1})"), FText::FromString(CurveTypeAsString), FText::AsNumber(static_cast<int32>(SupportedCurveType)));
+				
 			}
 		}
 		else
 		{
 			const FString CurveTypeAsString = GetCurveTypeValueName(SupportedCurveType);
-			ReportWarningf(LOCTEXT("InvalidCurveTypeWarning", "Invalid curve type provided: {0} ({1})"), FText::FromString(CurveTypeAsString), FText::AsNumber((int32)SupportedCurveType));
+			ReportWarningf(LOCTEXT("InvalidCurveTypeWarning", "Invalid curve type provided: {0} ({1})"), FText::FromString(CurveTypeAsString), FText::AsNumber(static_cast<int32>(SupportedCurveType)));
 		}
 	}
 	else
@@ -276,7 +320,7 @@ void UAnimDataController::FindOrAddCurveNamesOnSkeleton(USkeleton* Skeleton, ERa
 	{
 		if (IsSupportedCurveType(SupportedCurveType))
 		{
-			CONDITIONAL_BRACKET(LOCTEXT("FindOrAddRawCurveNames", "Updating Skeleton with Animation Curve Names"));
+			FBracket Bracket = ConditionalBracket(LOCTEXT("FindOrAddRawCurveNames", "Updating Skeleton with Animation Curve Names"), bShouldTransact);
 			switch (SupportedCurveType)
 			{
 			case ERawCurveTrackTypes::RCT_Float:
@@ -309,12 +353,16 @@ void UAnimDataController::FindOrAddCurveNamesOnSkeleton(USkeleton* Skeleton, ERa
 				}
 				break;
 			}
+			case ERawCurveTrackTypes::RCT_Vector:
+			default:
+				const FString CurveTypeAsString = GetCurveTypeValueName(SupportedCurveType);
+				ReportWarningf(LOCTEXT("InvalidCurveTypeWarning", "Invalid curve type provided: {0} ({1})"), FText::FromString(CurveTypeAsString), FText::AsNumber(static_cast<int32>(SupportedCurveType)));
 			}
 		}
 		else
 		{
 			const FString CurveTypeAsString = GetCurveTypeValueName(SupportedCurveType);
-			ReportWarningf(LOCTEXT("InvalidCurveTypeWarning", "Invalid curve type provided: {0} ({1})"), FText::FromString(CurveTypeAsString), FText::AsNumber((int32)SupportedCurveType));
+			ReportWarningf(LOCTEXT("InvalidCurveTypeWarning", "Invalid curve type provided: {0} ({1})"), FText::FromString(CurveTypeAsString), FText::AsNumber(static_cast<int32>(SupportedCurveType)));
 		}
 	}
 	else
@@ -325,7 +373,7 @@ void UAnimDataController::FindOrAddCurveNamesOnSkeleton(USkeleton* Skeleton, ERa
 
 bool UAnimDataController::RemoveBoneTracksMissingFromSkeleton(const USkeleton* Skeleton, bool bShouldTransact /*= true*/)
 {
-	if (!CheckOuterClass(UAnimSequence::StaticClass()))
+	if (!ModelInterface->GetAnimationSequence())
 	{
 		return false;
 	}
@@ -357,7 +405,7 @@ bool UAnimDataController::RemoveBoneTracksMissingFromSkeleton(const USkeleton* S
 
 		if (TracksToBeRemoved.Num() || TracksUpdated.Num())
 		{
-			CONDITIONAL_BRACKET(LOCTEXT("RemoveBoneTracksMissingFromSkeleton", "Validating Bone Animation Track Data against Skeleton"));
+			FBracket Bracket = ConditionalBracket(LOCTEXT("RemoveBoneTracksMissingFromSkeleton", "Validating Bone Animation Track Data against Skeleton"), bShouldTransact);
 			for (const FName& TrackName : TracksToBeRemoved)
 			{
 				RemoveBoneTrack(TrackName);
@@ -367,7 +415,7 @@ bool UAnimDataController::RemoveBoneTracksMissingFromSkeleton(const USkeleton* S
 			{
 				FAnimationTrackChangedPayload Payload;
 				Payload.Name = TrackName;
-				Model->Notify(EAnimDataModelNotifyType::TrackChanged, Payload);
+				Model->GetNotifier().Notify(EAnimDataModelNotifyType::TrackChanged, Payload);
 			}
 		}
 
@@ -407,7 +455,7 @@ void UAnimDataController::UpdateAttributesFromSkeleton(const USkeleton* Skeleton
 
 		if (ToRemoveIdentifiers.Num() || ToDuplicateIdentifiers.Num())
 		{
-			CONDITIONAL_BRACKET(LOCTEXT("VerifyAttributeBoneNames", "Remapping Animation Attribute Data"));
+			FBracket Bracket = ConditionalBracket(LOCTEXT("VerifyAttributeBoneNames", "Remapping Animation Attribute Data"), bShouldTransact);
 			for (const FAnimationAttributeIdentifier& Identifier : ToRemoveIdentifiers)
 			{
 				RemoveAttribute(Identifier);
@@ -415,8 +463,8 @@ void UAnimDataController::UpdateAttributesFromSkeleton(const USkeleton* Skeleton
 			
 			for (const TPair<FAnimationAttributeIdentifier, int32>& Pair : ToDuplicateIdentifiers)
 			{
-				FAnimationAttributeIdentifier NewIdentifier = Pair.Key;
-				NewIdentifier.BoneIndex = Pair.Value;
+				const FAnimationAttributeIdentifier& DuplicateIdentifier = Pair.Key;
+				FAnimationAttributeIdentifier NewIdentifier(DuplicateIdentifier.GetName(), Pair.Value, DuplicateIdentifier.GetBoneName(), DuplicateIdentifier.GetType());
 
 				DuplicateAttribute(Pair.Key, NewIdentifier);
 				RemoveAttribute(Pair.Key);
@@ -433,17 +481,17 @@ void UAnimDataController::ResetModel(bool bShouldTransact /*= true*/)
 {
 	ValidateModel();
 
-	CONDITIONAL_BRACKET(LOCTEXT("ResetModel", "Clearing Animation Data"));
+	FBracket Bracket = ConditionalBracket(LOCTEXT("ResetModel", "Clearing Animation Data"), bShouldTransact);
 
 	RemoveAllBoneTracks(bShouldTransact);
 
 	RemoveAllCurvesOfType(ERawCurveTrackTypes::RCT_Float, bShouldTransact);
 	RemoveAllCurvesOfType(ERawCurveTrackTypes::RCT_Transform, bShouldTransact);
 
-	SetPlayLength(MINIMUM_ANIMATION_LENGTH, bShouldTransact);
-	SetFrameRate(FFrameRate(30,1), bShouldTransact);
+	SetFrameRate(UAnimationSettings::Get()->GetDefaultFrameRate(), bShouldTransact);
+	SetNumberOfFrames(1, bShouldTransact);
 
-	Model->Notify(EAnimDataModelNotifyType::Reset);
+	Model->GetNotifier().Notify(EAnimDataModelNotifyType::Reset);
 }
 
 bool UAnimDataController::AddCurve(const FAnimationCurveIdentifier& CurveId, int32 CurveFlags /*= EAnimAssetCurveFlags::AACF_Editable*/, bool bShouldTransact /*= true*/)
@@ -455,7 +503,7 @@ bool UAnimDataController::AddCurve(const FAnimationCurveIdentifier& CurveId, int
 		{
 			if (!Model->FindCurve(CurveId))
 			{
-				CONDITIONAL_TRANSACTION(LOCTEXT("AddRawCurve", "Adding Animation Curve"));
+				FTransaction Transaction = ConditionalTransaction(LOCTEXT("AddRawCurve", "Adding Animation Curve"), bShouldTransact);
 
 				FCurveAddedPayload Payload;
 				Payload.Identifier = CurveId;
@@ -473,23 +521,27 @@ bool UAnimDataController::AddCurve(const FAnimationCurveIdentifier& CurveId, int
 				case ERawCurveTrackTypes::RCT_Float:
 					AddNewCurve(Model->CurveData.FloatCurves);
 					break;
+				case ERawCurveTrackTypes::RCT_Vector:
+				default:
+					const FString CurveTypeAsString = GetCurveTypeValueName(CurveId.CurveType);
+					ReportWarningf(LOCTEXT("InvalidCurveTypeWarning", "Invalid curve type provided: {0} ({1})"), FText::FromString(CurveTypeAsString), FText::AsNumber(static_cast<int32>(CurveId.CurveType)));
 				}
 
-				CONDITIONAL_ACTION(UE::Anim::FRemoveCurveAction, CurveId);
-				Model->Notify(EAnimDataModelNotifyType::CurveAdded, Payload);
+				ConditionalAction<UE::Anim::FRemoveCurveAction>(bShouldTransact, CurveId);
+				Model->GetNotifier().Notify(EAnimDataModelNotifyType::CurveAdded, Payload);
 
 				return true;
 			}
 			else
 			{
 				const FString CurveTypeAsString = GetCurveTypeValueName(CurveId.CurveType);
-				ReportWarningf(LOCTEXT("ExistingCurveNameWarning", "Curve with name {0} and type {1} ({2}) already exists"), FText::FromName(CurveId.InternalName.DisplayName), FText::FromString(CurveTypeAsString), FText::AsNumber((int32)CurveId.CurveType));
+				ReportWarningf(LOCTEXT("ExistingCurveNameWarning", "Curve with name {0} and type {1} ({2}) already exists"), FText::FromName(CurveId.InternalName.DisplayName), FText::FromString(CurveTypeAsString), FText::AsNumber(static_cast<int32>(CurveId.CurveType)));
 			}			
 		}
 		else 
 		{
 			const FString CurveTypeAsString = GetCurveTypeValueName(CurveId.CurveType);
-			ReportWarningf(LOCTEXT("InvalidCurveTypeWarning", "Invalid curve type provided: {0} ({1})"), FText::FromString(CurveTypeAsString), FText::AsNumber((int32)CurveId.CurveType));
+			ReportWarningf(LOCTEXT("InvalidCurveTypeWarning", "Invalid curve type provided: {0} ({1})"), FText::FromString(CurveTypeAsString), FText::AsNumber(static_cast<int32>(CurveId.CurveType)));
 		}		
 	}
 	else
@@ -517,7 +569,7 @@ bool UAnimDataController::DuplicateCurve(const FAnimationCurveIdentifier& CopyCu
 				{
 					if (!Model->FindCurve(NewCurveId))
 					{
-						CONDITIONAL_TRANSACTION(LOCTEXT("CopyRawCurve", "Duplicating Animation Curve"));
+						FTransaction Transaction = ConditionalTransaction(LOCTEXT("CopyRawCurve", "Duplicating Animation Curve"), bShouldTransact);
 
 						auto DuplicateCurve = [NewCurveName = NewCurveId.InternalName](auto& CurveDataArray, const auto& SourceCurve)
 						{
@@ -533,33 +585,37 @@ bool UAnimDataController::DuplicateCurve(const FAnimationCurveIdentifier& CopyCu
 						case ERawCurveTrackTypes::RCT_Float:
 							DuplicateCurve(Model->CurveData.FloatCurves, Model->GetFloatCurve(CopyCurveId));
 							break;
+						case ERawCurveTrackTypes::RCT_Vector:
+						default:
+							const FString CurveTypeAsString = GetCurveTypeValueName(SupportedCurveType);
+							ReportWarningf(LOCTEXT("InvalidCurveTypeWarning", "Invalid curve type provided: {0} ({1})"), FText::FromString(CurveTypeAsString), FText::AsNumber(static_cast<int32>(SupportedCurveType)));
 						}
 
 						FCurveAddedPayload Payload;
 						Payload.Identifier = NewCurveId;
-						Model->Notify(EAnimDataModelNotifyType::CurveAdded, Payload);
+						Model->GetNotifier().Notify(EAnimDataModelNotifyType::CurveAdded, Payload);
 
-						CONDITIONAL_ACTION(UE::Anim::FRemoveCurveAction, NewCurveId);
+						ConditionalAction<UE::Anim::FRemoveCurveAction>(bShouldTransact, NewCurveId);
 
 						return true;
 					}
 					else
 					{
 						const FString CurveTypeAsString = GetCurveTypeValueName(NewCurveId.CurveType);
-						ReportWarningf(LOCTEXT("ExistingCurveNameWarning", "Curve with name {0} and type {1} ({2}) already exists"), FText::FromName(NewCurveId.InternalName.DisplayName), FText::FromString(CurveTypeAsString), FText::AsNumber((int32)NewCurveId.CurveType));
+						ReportWarningf(LOCTEXT("ExistingCurveNameWarning", "Curve with name {0} and type {1} ({2}) already exists"), FText::FromName(NewCurveId.InternalName.DisplayName), FText::FromString(CurveTypeAsString), FText::AsNumber(static_cast<int32>(NewCurveId.CurveType)));
 					}
 				}
 				else
 				{
 					const FString CurveTypeAsString = GetCurveTypeValueName(CopyCurveId.CurveType);
-					ReportWarningf(LOCTEXT("CurveNameToDuplicateNotFoundWarning", "Could not find curve with name {0} and type {1} ({2}) for duplication"), FText::FromName(NewCurveId.InternalName.DisplayName), FText::FromString(CurveTypeAsString), FText::AsNumber((int32)NewCurveId.CurveType));
+					ReportWarningf(LOCTEXT("CurveNameToDuplicateNotFoundWarning", "Could not find curve with name {0} and type {1} ({2}) for duplication"), FText::FromName(NewCurveId.InternalName.DisplayName), FText::FromString(CurveTypeAsString), FText::AsNumber(static_cast<int32>(NewCurveId.CurveType)));
 				}
 			}
 		}
 		else
 		{
 			const FString CurveTypeAsString = GetCurveTypeValueName(SupportedCurveType);
-			ReportWarningf(LOCTEXT("InvalidCurveTypeWarning", "Invalid curve type provided: {0} ({1})"), FText::FromString(CurveTypeAsString), FText::AsNumber((int32)SupportedCurveType));
+			ReportWarningf(LOCTEXT("InvalidCurveTypeWarning", "Invalid curve type provided: {0} ({1})"), FText::FromString(CurveTypeAsString), FText::AsNumber(static_cast<int32>(SupportedCurveType)));
 		}
 	}
 
@@ -576,32 +632,35 @@ bool UAnimDataController::RemoveCurve(const FAnimationCurveIdentifier& CurveId,
 	{
 		if (IsSupportedCurveType(CurveId.CurveType))
 		{
-			const FAnimCurveBase* Curve = Model->FindCurve(CurveId);
-			if (Curve)
+			if (Model->FindCurve(CurveId) != nullptr)
 			{
-				CONDITIONAL_TRANSACTION(LOCTEXT("RemoveCurve", "Removing Animation Curve"));
+				FTransaction Transaction = ConditionalTransaction(LOCTEXT("RemoveCurve", "Removing Animation Curve"), bShouldTransact);
 
 				switch (SupportedCurveType)
 				{
 					case ERawCurveTrackTypes::RCT_Transform:
 					{
 						const FTransformCurve& TransformCurve = Model->GetTransformCurve(CurveId);
-						CONDITIONAL_ACTION(UE::Anim::FAddTransformCurveAction, CurveId, TransformCurve.GetCurveTypeFlags(), TransformCurve);
+						ConditionalAction<UE::Anim::FAddTransformCurveAction>(bShouldTransact, CurveId, TransformCurve.GetCurveTypeFlags(), TransformCurve);
 						Model->CurveData.TransformCurves.RemoveAll([Name = TransformCurve.Name](const FTransformCurve& ToRemoveCurve) { return ToRemoveCurve.Name == Name; });
 						break;
 					}
 					case ERawCurveTrackTypes::RCT_Float:
 					{
 						const FFloatCurve& FloatCurve = Model->GetFloatCurve(CurveId);
-						CONDITIONAL_ACTION(UE::Anim::FAddFloatCurveAction, CurveId, FloatCurve.GetCurveTypeFlags(), FloatCurve.FloatCurve.GetConstRefOfKeys(), FloatCurve.Color);
+						ConditionalAction<UE::Anim::FAddFloatCurveAction>(bShouldTransact, CurveId, FloatCurve.GetCurveTypeFlags(), FloatCurve.FloatCurve.GetConstRefOfKeys(), FloatCurve.Color);
 						Model->CurveData.FloatCurves.RemoveAll([Name = FloatCurve.Name](const FFloatCurve& ToRemoveCurve) { return ToRemoveCurve.Name == Name; });
 						break;
 					}
+					case ERawCurveTrackTypes::RCT_Vector:
+					default:
+						const FString CurveTypeAsString = GetCurveTypeValueName(SupportedCurveType);
+						ReportWarningf(LOCTEXT("InvalidCurveTypeWarning", "Invalid curve type provided: {0} ({1})"), FText::FromString(CurveTypeAsString), FText::AsNumber(static_cast<int32>(SupportedCurveType)));
 				}
 
 				FCurveRemovedPayload Payload;
 				Payload.Identifier = CurveId;
-				Model->Notify(EAnimDataModelNotifyType::CurveRemoved, Payload);
+				Model->GetNotifier().Notify(EAnimDataModelNotifyType::CurveRemoved, Payload);
 
 				return true;
 			}
@@ -614,7 +673,7 @@ bool UAnimDataController::RemoveCurve(const FAnimationCurveIdentifier& CurveId,
 		else
 		{
 			const FString CurveTypeAsString = GetCurveTypeValueName(CurveId.CurveType);
-			ReportWarningf(LOCTEXT("InvalidCurveTypeWarning", "Invalid curve type provided: {0} ({1})"), FText::FromString(CurveTypeAsString), FText::AsNumber((int32)CurveId.CurveType));
+			ReportWarningf(LOCTEXT("InvalidCurveTypeWarning", "Invalid curve type provided: {0} ({1})"), FText::FromString(CurveTypeAsString), FText::AsNumber(static_cast<int32>(CurveId.CurveType)));
 		}
 	}
 
@@ -625,7 +684,7 @@ void UAnimDataController::RemoveAllCurvesOfType(ERawCurveTrackTypes SupportedCur
 {
 	ValidateModel();
 
-	CONDITIONAL_BRACKET(LOCTEXT("DeleteAllRawCurve", "Deleting All Animation Curve"));
+	FBracket Bracket = ConditionalBracket(LOCTEXT("DeleteAllRawCurve", "Deleting All Animation Curve"), bShouldTransact);
 	switch (SupportedCurveType)
 	{
 	case ERawCurveTrackTypes::RCT_Transform:
@@ -649,7 +708,7 @@ void UAnimDataController::RemoveAllCurvesOfType(ERawCurveTrackTypes SupportedCur
 	case ERawCurveTrackTypes::RCT_Vector:
 	default:
 		const FString CurveTypeAsString = GetCurveTypeValueName(SupportedCurveType);
-		ReportWarningf(LOCTEXT("InvalidCurveTypeWarning", "Invalid curve type provided: {0} ({1})"), FText::FromString(CurveTypeAsString), FText::AsNumber((int32)SupportedCurveType));
+		ReportWarningf(LOCTEXT("InvalidCurveTypeWarning", "Invalid curve type provided: {0} ({1})"), FText::FromString(CurveTypeAsString), FText::AsNumber(static_cast<int32>(SupportedCurveType)));
 	}
 
 }
@@ -658,7 +717,7 @@ bool UAnimDataController::SetCurveFlag(const FAnimationCurveIdentifier& CurveId,
 {
 	ValidateModel();
 
-	ERawCurveTrackTypes SupportedCurveType = CurveId.CurveType;
+	const ERawCurveTrackTypes SupportedCurveType = CurveId.CurveType;
 
 	FAnimCurveBase* Curve = nullptr;
 
@@ -670,14 +729,19 @@ bool UAnimDataController::SetCurveFlag(const FAnimationCurveIdentifier& CurveId,
 	{
 		Curve = Model->FindMutableTransformCurveById(CurveId);
 	}
+	else
+	{
+		const FString CurveTypeAsString = GetCurveTypeValueName(SupportedCurveType);
+		ReportWarningf(LOCTEXT("InvalidCurveTypeWarning", "Invalid curve type provided: {0} ({1})"), FText::FromString(CurveTypeAsString), FText::AsNumber(static_cast<int32>(SupportedCurveType)));
+	}
 	
 	if (Curve)
 	{
-		CONDITIONAL_TRANSACTION(LOCTEXT("SetCurveFlag", "Setting Raw Curve Flag"));
+		FTransaction Transaction = ConditionalTransaction(LOCTEXT("SetCurveFlag", "Setting Raw Curve Flag"), bShouldTransact);
 
 		const int32 CurrentFlags = Curve->GetCurveTypeFlags();
 
-		CONDITIONAL_ACTION(UE::Anim::FSetCurveFlagsAction, CurveId, CurrentFlags, SupportedCurveType);
+		ConditionalAction<UE::Anim::FSetCurveFlagsAction>(bShouldTransact, CurveId, CurrentFlags, SupportedCurveType);
 
 		FCurveFlagsChangedPayload Payload;
 		Payload.Identifier = CurveId;
@@ -685,7 +749,7 @@ bool UAnimDataController::SetCurveFlag(const FAnimationCurveIdentifier& CurveId,
 
 		Curve->SetCurveTypeFlag(Flag, bState);
 
-		Model->Notify(EAnimDataModelNotifyType::CurveFlagsChanged, Payload);
+		Model->GetNotifier().Notify(EAnimDataModelNotifyType::CurveFlagsChanged, Payload);
 
 		return true;
 	}
@@ -704,8 +768,7 @@ bool UAnimDataController::SetCurveFlags(const FAnimationCurveIdentifier& CurveId
 
 	FAnimCurveBase* Curve = nullptr;
 
-	ERawCurveTrackTypes SupportedCurveType = CurveId.CurveType;
-
+	const ERawCurveTrackTypes SupportedCurveType = CurveId.CurveType;
 	if (SupportedCurveType == ERawCurveTrackTypes::RCT_Float)
 	{
 		Curve = Model->FindMutableFloatCurveById(CurveId);
@@ -714,14 +777,19 @@ bool UAnimDataController::SetCurveFlags(const FAnimationCurveIdentifier& CurveId
 	{
 		Curve = Model->FindMutableTransformCurveById(CurveId);
 	}
+	else
+	{
+		const FString CurveTypeAsString = GetCurveTypeValueName(SupportedCurveType);
+		ReportWarningf(LOCTEXT("InvalidCurveTypeWarning", "Invalid curve type provided: {0} ({1})"), FText::FromString(CurveTypeAsString), FText::AsNumber(static_cast<int32>(SupportedCurveType)));
+	}
 
 	if (Curve)
 	{
-		CONDITIONAL_TRANSACTION(LOCTEXT("SetRawCurveFlag", "Setting Raw Curve Flags"));
+		FTransaction Transaction = ConditionalTransaction(LOCTEXT("SetRawCurveFlag", "Setting Raw Curve Flags"), bShouldTransact);
 
 		const int32 CurrentFlags = Curve->GetCurveTypeFlags();
 
-		CONDITIONAL_ACTION(UE::Anim::FSetCurveFlagsAction, CurveId, CurrentFlags, SupportedCurveType);
+		ConditionalAction<UE::Anim::FSetCurveFlagsAction>(bShouldTransact, CurveId, CurrentFlags, SupportedCurveType);
 
 		FCurveFlagsChangedPayload Payload;
 		Payload.Identifier = CurveId;
@@ -729,7 +797,7 @@ bool UAnimDataController::SetCurveFlags(const FAnimationCurveIdentifier& CurveId
 
 		Curve->SetCurveTypeFlags(Flags);
 
-		Model->Notify(EAnimDataModelNotifyType::CurveFlagsChanged, Payload);
+		Model->GetNotifier().Notify(EAnimDataModelNotifyType::CurveFlagsChanged, Payload);
 
 		return true;
 	}
@@ -748,11 +816,9 @@ bool UAnimDataController::SetTransformCurveKeys(const FAnimationCurveIdentifier&
 
 	if (TransformValues.Num() == TimeKeys.Num())
 	{
-		FTransformCurve* Curve = Model->FindMutableTransformCurveById(CurveId);
-
-		if (Curve)
+		if (Model->FindMutableTransformCurveById(CurveId) != nullptr)
 		{
-			CONDITIONAL_BRACKET(LOCTEXT("SetTransformCurveKeys_Bracket", "Setting Transform Curve Keys"));
+			FBracket Bracket = ConditionalBracket(LOCTEXT("SetTransformCurveKeys_Bracket", "Setting Transform Curve Keys"), bShouldTransact);
 			
 			struct FKeys
 			{
@@ -798,11 +864,11 @@ bool UAnimDataController::SetTransformCurveKeys(const FAnimationCurveIdentifier&
 			
 			for (int32 SubCurveIndex = 0; SubCurveIndex < 3; ++SubCurveIndex)
 			{
-				const ETransformCurveChannel Channel = (ETransformCurveChannel)SubCurveIndex;
-				FKeys* CurveKeys = SubCurveKeys[SubCurveIndex];
+				const ETransformCurveChannel Channel = static_cast<ETransformCurveChannel>(SubCurveIndex);
+				const FKeys* CurveKeys = SubCurveKeys[SubCurveIndex];
 				for (int32 ChannelIndex = 0; ChannelIndex < 3; ++ChannelIndex)
 				{
-					const EVectorCurveChannel Axis = (EVectorCurveChannel)ChannelIndex;
+					const EVectorCurveChannel Axis = static_cast<EVectorCurveChannel>(ChannelIndex);
 					FAnimationCurveIdentifier TargetCurveIdentifier = CurveId;
 					UAnimationCurveIdentifierExtensions::GetTransformChildCurveIdentifier(TargetCurveIdentifier, Channel, Axis);
 					SetCurveKeys(TargetCurveIdentifier, CurveKeys->ChannelKeys[ChannelIndex], bShouldTransact);
@@ -830,11 +896,9 @@ bool UAnimDataController::SetTransformCurveKey(const FAnimationCurveIdentifier&
 {
 	ValidateModel();
 
-	FTransformCurve* Curve = Model->FindMutableTransformCurveById(CurveId);
-
-	if (Curve)
+	if (Model->FindMutableTransformCurveById(CurveId) != nullptr)
 	{
-		CONDITIONAL_BRACKET(LOCTEXT("AddTransformCurveKey_Bracket", "Setting Transform Curve Key"));
+		FBracket Bracket = ConditionalBracket(LOCTEXT("AddTransformCurveKey_Bracket", "Setting Transform Curve Key"), bShouldTransact);
 		struct FKeys
 		{
 			FRichCurveKey ChannelKeys[3];
@@ -861,11 +925,11 @@ bool UAnimDataController::SetTransformCurveKey(const FAnimationCurveIdentifier&
 		
 		for (int32 SubCurveIndex = 0; SubCurveIndex < 3; ++SubCurveIndex)
 		{
-			const ETransformCurveChannel Channel = (ETransformCurveChannel)SubCurveIndex;
+			const ETransformCurveChannel Channel = static_cast<ETransformCurveChannel>(SubCurveIndex);
 			const FKeys& VectorCurveKeys = VectorKeys[SubCurveIndex];
 			for (int32 ChannelIndex = 0; ChannelIndex < 3; ++ChannelIndex)
 			{
-				const EVectorCurveChannel Axis = (EVectorCurveChannel)ChannelIndex;
+				const EVectorCurveChannel Axis = static_cast<EVectorCurveChannel>(ChannelIndex);
 				FAnimationCurveIdentifier TargetCurveIdentifier = CurveId;
 				UAnimationCurveIdentifierExtensions::GetTransformChildCurveIdentifier(TargetCurveIdentifier, Channel, Axis);
 				SetCurveKey(TargetCurveIdentifier, VectorCurveKeys.ChannelKeys[ChannelIndex], bShouldTransact);
@@ -887,30 +951,27 @@ bool UAnimDataController::RemoveTransformCurveKey(const FAnimationCurveIdentifie
 {
 	ValidateModel();
 
-	FTransformCurve* TransformCurve = Model->FindMutableTransformCurveById(CurveId);
-	if (TransformCurve)
+	if (Model->FindMutableTransformCurveById(CurveId))
 	{
 		const FString BaseCurveName = CurveId.InternalName.DisplayName.ToString();
 		const TArray<FString> SubCurveNames = { TEXT( "Translation"), TEXT( "Rotation"), TEXT( "Scale") };
 		const TArray<FString> ChannelCurveNames = { TEXT("X"), TEXT("Y"), TEXT("Z") };
 
-		CONDITIONAL_BRACKET(LOCTEXT("RemoveTransformCurveKey_Bracket", "Deleting Animation Transform Curve Key"));
+		FBracket Bracket = ConditionalBracket(LOCTEXT("RemoveTransformCurveKey_Bracket", "Deleting Animation Transform Curve Key"), bShouldTransact);
 		
 		for (int32 SubCurveIndex = 0; SubCurveIndex < 3; ++SubCurveIndex)
 		{
-			const ETransformCurveChannel Channel = (ETransformCurveChannel)SubCurveIndex;
+			const ETransformCurveChannel Channel = static_cast<ETransformCurveChannel>(SubCurveIndex);
 			for (int32 ChannelIndex = 0; ChannelIndex < 3; ++ChannelIndex)
 			{
-				const EVectorCurveChannel Axis = (EVectorCurveChannel)ChannelIndex;
+				const EVectorCurveChannel Axis = static_cast<EVectorCurveChannel>(ChannelIndex);
 				FAnimationCurveIdentifier TargetCurveIdentifier = CurveId;
 				UAnimationCurveIdentifierExtensions::GetTransformChildCurveIdentifier(TargetCurveIdentifier, Channel, Axis);
 				RemoveCurveKey(TargetCurveIdentifier, Time, bShouldTransact);
 			}
 		}
 
-
 		return true;
-
 	}
 	else
 	{
@@ -930,10 +991,9 @@ bool UAnimDataController::RenameCurve(const FAnimationCurveIdentifier& CurveToRe
 		{
 			if (CurveToRenameId.CurveType == NewCurveId.CurveType)
 			{
-				FAnimCurveBase* Curve = Model->FindMutableCurveById(CurveToRenameId);
-				if (Curve)
+				if (FAnimCurveBase* Curve = Model->FindMutableCurveById(CurveToRenameId))
 				{
-					CONDITIONAL_TRANSACTION(LOCTEXT("RenameCurve", "Renaming Curve"));
+					FTransaction Transaction = ConditionalTransaction(LOCTEXT("RenameCurve", "Renaming Curve"), bShouldTransact);
 
 					FCurveRenamedPayload Payload;
 					Payload.Identifier = FAnimationCurveIdentifier(Curve->Name, CurveToRenameId.CurveType);
@@ -941,9 +1001,9 @@ bool UAnimDataController::RenameCurve(const FAnimationCurveIdentifier& CurveToRe
 					Curve->Name = NewCurveId.InternalName;
 					Payload.NewIdentifier = NewCurveId;
 
-					CONDITIONAL_ACTION(UE::Anim::FRenameCurveAction, NewCurveId, CurveToRenameId);
+					ConditionalAction<UE::Anim::FRenameCurveAction>(bShouldTransact, NewCurveId, CurveToRenameId);
 
-					Model->Notify(EAnimDataModelNotifyType::CurveRenamed, Payload);
+					Model->GetNotifier().Notify(EAnimDataModelNotifyType::CurveRenamed, Payload);
 
 					return true;
 				}
@@ -983,18 +1043,17 @@ bool UAnimDataController::SetCurveColor(const FAnimationCurveIdentifier& CurveId
 	{
 		if (CurveId.CurveType == ERawCurveTrackTypes::RCT_Float)
 		{
-			FFloatCurve* Curve = Model->FindMutableFloatCurveById(CurveId);
-			if (Curve)
+			if (FFloatCurve* Curve = Model->FindMutableFloatCurveById(CurveId))
 			{
-				CONDITIONAL_TRANSACTION(LOCTEXT("ChangingCurveColor", "Changing Curve Color"));
+				FTransaction Transaction = ConditionalTransaction(LOCTEXT("ChangingCurveColor", "Changing Curve Color"), bShouldTransact);
 
-				CONDITIONAL_ACTION(UE::Anim::FSetCurveColorAction, CurveId, Curve->Color);
+				ConditionalAction<UE::Anim::FSetCurveColorAction>(bShouldTransact, CurveId, Curve->Color);
 
 				Curve->Color = Color;
 
 				FCurveChangedPayload Payload;
 				Payload.Identifier = CurveId;
-				Model->Notify(EAnimDataModelNotifyType::CurveColorChanged, Payload);
+				Model->GetNotifier().Notify(EAnimDataModelNotifyType::CurveColorChanged, Payload);
 
 				return true;				
 			}
@@ -1021,13 +1080,12 @@ bool UAnimDataController::ScaleCurve(const FAnimationCurveIdentifier& CurveId, f
 {
 	ValidateModel();
 
-	ERawCurveTrackTypes SupportedCurveType = CurveId.CurveType;
+	const ERawCurveTrackTypes SupportedCurveType = CurveId.CurveType;
 	if (SupportedCurveType == ERawCurveTrackTypes::RCT_Float)
 	{
-		FFloatCurve* Curve = Model->FindMutableFloatCurveById(CurveId);
-		if (Curve)
+		if (FFloatCurve* Curve = Model->FindMutableFloatCurveById(CurveId))
 		{
-			CONDITIONAL_TRANSACTION(LOCTEXT("ScalingCurve", "Scaling Curve"));
+			FTransaction Transaction = ConditionalTransaction(LOCTEXT("ScalingCurve", "Scaling Curve"), bShouldTransact);
 
 			Curve->FloatCurve.ScaleCurve(Origin, Factor);
 
@@ -1036,9 +1094,9 @@ bool UAnimDataController::ScaleCurve(const FAnimationCurveIdentifier& CurveId, f
 			Payload.Factor = Factor;
 			Payload.Origin = Origin;
 			
-			CONDITIONAL_ACTION(UE::Anim::FScaleCurveAction, CurveId, Origin, 1.0f / Factor, SupportedCurveType);
+			ConditionalAction<UE::Anim::FScaleCurveAction>(bShouldTransact, CurveId, Origin, 1.0f / Factor, SupportedCurveType);
 
-			Model->Notify(EAnimDataModelNotifyType::CurveScaled, Payload);
+			Model->GetNotifier().Notify(EAnimDataModelNotifyType::CurveScaled, Payload);
 
 			return true;
 		}
@@ -1059,9 +1117,7 @@ bool UAnimDataController::SetCurveKey(const FAnimationCurveIdentifier& CurveId,
 {
 	ValidateModel();
 
-	FRichCurve* RichCurve = Model->GetMutableRichCurve(CurveId);
-	ERawCurveTrackTypes SupportedCurveType = CurveId.CurveType;
-	if (RichCurve)
+	if (FRichCurve* RichCurve = Model->GetMutableRichCurve(CurveId))
 	{
 		FCurveChangedPayload Payload;
 		Payload.Identifier = CurveId;
@@ -1070,25 +1126,25 @@ bool UAnimDataController::SetCurveKey(const FAnimationCurveIdentifier& CurveId,
 		const FKeyHandle Handle = RichCurve->FindKey(Key.Time, 0.f);
 		if (Handle != FKeyHandle::Invalid())
 		{
-			CONDITIONAL_TRANSACTION(LOCTEXT("SetNamedCurveKey", "Setting Curve Key"));
+			FTransaction Transaction = ConditionalTransaction(LOCTEXT("SetNamedCurveKey", "Setting Curve Key"), bShouldTransact);
 			// Cache old value for action
 			const FRichCurveKey CurrentKey = RichCurve->GetKey(Handle);
-			CONDITIONAL_ACTION(UE::Anim::FSetRichCurveKeyAction, CurveId, CurrentKey);
+			ConditionalAction<UE::Anim::FSetRichCurveKeyAction>(bShouldTransact, CurveId, CurrentKey);
 
 			// Set the new value
 			RichCurve->SetKeyValue(Handle, Key.Value);
 
-			Model->Notify(EAnimDataModelNotifyType::CurveChanged, Payload);
+			Model->GetNotifier().Notify(EAnimDataModelNotifyType::CurveChanged, Payload);
 		}
 		else
 		{
-			CONDITIONAL_TRANSACTION(LOCTEXT("AddNamedCurveKey", "Adding Curve Key"));
-			CONDITIONAL_ACTION(UE::Anim::FRemoveRichCurveKeyAction, CurveId, Key.Time);
+			FTransaction Transaction = ConditionalTransaction(LOCTEXT("AddNamedCurveKey", "Adding Curve Key"), bShouldTransact);
+			ConditionalAction<UE::Anim::FRemoveRichCurveKeyAction>(bShouldTransact, CurveId, Key.Time);
 
 			// Add the new key
 			RichCurve->AddKey(Key.Time, Key.Value);
 
-			Model->Notify(EAnimDataModelNotifyType::CurveChanged, Payload);
+			Model->GetNotifier().Notify(EAnimDataModelNotifyType::CurveChanged, Payload);
 		}
 
 		return true;
@@ -1101,9 +1157,7 @@ bool UAnimDataController::RemoveCurveKey(const FAnimationCurveIdentifier& CurveI
 {
 	ValidateModel();
 
-	FRichCurve* RichCurve = Model->GetMutableRichCurve(CurveId);
-	ERawCurveTrackTypes SupportedCurveType = CurveId.CurveType;
-	if (RichCurve)
+	if (FRichCurve* RichCurve = Model->GetMutableRichCurve(CurveId))
 	{
 		FCurveChangedPayload Payload;
 		Payload.Identifier = CurveId;
@@ -1112,15 +1166,15 @@ bool UAnimDataController::RemoveCurveKey(const FAnimationCurveIdentifier& CurveI
 		const FKeyHandle Handle = RichCurve->FindKey(Time, 0.f);
 		if (Handle != FKeyHandle::Invalid())
 		{
-			CONDITIONAL_TRANSACTION(LOCTEXT("RemoveNamedCurveKey", "Removing Curve Key"));
+			FTransaction Transaction = ConditionalTransaction(LOCTEXT("RemoveNamedCurveKey", "Removing Curve Key"), bShouldTransact);
 
 			// Cached current value for action
 			const FRichCurveKey CurrentKey = RichCurve->GetKey(Handle);
-			CONDITIONAL_ACTION(UE::Anim::FAddRichCurveKeyAction, CurveId, CurrentKey);
+			ConditionalAction<UE::Anim::FAddRichCurveKeyAction>(bShouldTransact, CurveId, CurrentKey);
 
 			RichCurve->DeleteKey(Handle);
 
-			Model->Notify(EAnimDataModelNotifyType::CurveChanged, Payload);
+			Model->GetNotifier().Notify(EAnimDataModelNotifyType::CurveChanged, Payload);
 
 			return true;
 		}
@@ -1138,19 +1192,17 @@ bool UAnimDataController::SetCurveKeys(const FAnimationCurveIdentifier& CurveId,
 {
 	ValidateModel();
 
-	FRichCurve* RichCurve = Model->GetMutableRichCurve(CurveId);
-	ERawCurveTrackTypes SupportedCurveType = CurveId.CurveType;
-	if (RichCurve)
+	if (FRichCurve* RichCurve = Model->GetMutableRichCurve(CurveId))
 	{
-		CONDITIONAL_TRANSACTION(LOCTEXT("SettingNamedCurveKeys", "Setting Curve Keys"));
-		CONDITIONAL_ACTION(UE::Anim::FSetRichCurveKeysAction, CurveId, RichCurve->GetConstRefOfKeys());
+		FTransaction Transaction = ConditionalTransaction(LOCTEXT("SettingNamedCurveKeys", "Setting Curve Keys"), bShouldTransact);
+		ConditionalAction<UE::Anim::FSetRichCurveKeysAction>(bShouldTransact, CurveId, RichCurve->GetConstRefOfKeys());
 
 		// Set rich curve values
 		RichCurve->SetKeys(CurveKeys);
 
 		FCurveChangedPayload Payload;
 		Payload.Identifier = CurveId;
-		Model->Notify(EAnimDataModelNotifyType::CurveChanged, Payload);
+		Model->GetNotifier().Notify(EAnimDataModelNotifyType::CurveChanged, Payload);
 
 		return true;
 	}
@@ -1161,165 +1213,42 @@ bool UAnimDataController::SetCurveKeys(const FAnimationCurveIdentifier& CurveId,
 void UAnimDataController::NotifyPopulated()
 {
 	ValidateModel();
-	Model->Notify(EAnimDataModelNotifyType::Populated);
+
+	Model->bPopulated = true;
+	Model->GetNotifier().Notify(EAnimDataModelNotifyType::Populated);
 }
 
 void UAnimDataController::NotifyBracketOpen()
 {
 	ValidateModel();
-	Model->Notify(EAnimDataModelNotifyType::BracketOpened);
+	Model->GetNotifier().Notify(EAnimDataModelNotifyType::BracketOpened);
 }
 
 void UAnimDataController::NotifyBracketClosed()
 {
 	ValidateModel();
-	Model->Notify(EAnimDataModelNotifyType::BracketClosed);
-}
-
-const bool UAnimDataController::IsSupportedCurveType(ERawCurveTrackTypes CurveType) const
-{
-	const TArray<ERawCurveTrackTypes> SupportedTypes = { ERawCurveTrackTypes::RCT_Float, ERawCurveTrackTypes::RCT_Transform };
-	return SupportedTypes.Contains(CurveType);
-}
-
-void UAnimDataController::ValidateModel() const
-{
-	checkf(Model != nullptr, TEXT("Invalid Model"));
+	Model->GetNotifier().Notify(EAnimDataModelNotifyType::BracketClosed);
 }
 
 void UAnimDataController::ResizePlayLength(float Length, float T0, float T1, bool bShouldTransact)
 {
-	const TRange<float> PlayRange(TRange<float>::BoundsType::Inclusive(0.f), TRange<float>::BoundsType::Inclusive(Model->PlayLength));
-	if (!FMath::IsNearlyZero(Length) && Length > 0.f)
-	{
-		if (Length != Model->PlayLength)
-		{
-			// Ensure that T0 is within the curent play range
-			if (PlayRange.Contains(T0))
-			{
-				// Ensure that the start and end length of either removal or insertion are valid
-				if (T0 < T1)
-				{
-					CONDITIONAL_TRANSACTION(LOCTEXT("ResizePlayLength", "Resizing Play Length"));
-
-					FSequenceLengthChangedPayload Payload;
-					Payload.T0 = T0;
-					Payload.T1 = T1;
-					Payload.PreviousLength = Model->PlayLength;
-
-					CONDITIONAL_ACTION(UE::Anim::FResizePlayLengthAction, Model, T0, T1);
-
-					Model->PlayLength = Length;
-
-					Model->NumberOfFrames = Model->FrameRate.AsFrameTime(Model->PlayLength).RoundToFrame().Value;
-					Model->NumberOfKeys = Model->NumberOfFrames + 1;
-	
-					Model->Notify<FSequenceLengthChangedPayload>(EAnimDataModelNotifyType::SequenceLengthChanged, Payload);
-				}
-				else
-				{
-					ReportErrorf(LOCTEXT("InvalidEndTimeError", "Invalid T1, smaller that T0 value: T0 {0}, T1 {1}"), FText::AsNumber(T0), FText::AsNumber(T1));
-				}
-			}
-			else
-			{
-				ReportErrorf(LOCTEXT("InvalidStartTimeError", "Invalid T0, not within existing play range: T0 {0}, Play Length {1}"), FText::AsNumber(T0), FText::AsNumber(Model->PlayLength));
-			}
-		}
-		else
-		{
-			ReportWarningf(LOCTEXT("SamePlayLengthWarning", "New play length is same as existing one: {0} seconds"), FText::AsNumber(Length));
-		}
-	}
-	else
-	{
-		ReportErrorf(LOCTEXT("InvalidPlayLengthError", "Invalid play length value provided: {0} seconds"), FText::AsNumber(Length));
-	}
-}
-
-void UAnimDataController::ReportWarning(const FText& InMessage) const
-{
-	FString Message = InMessage.ToString();
-	if (Model != nullptr)
-	{
-		if (UPackage* Package = Cast<UPackage>(Model->GetOutermost()))
-		{
-			Message = FString::Printf(TEXT("%s : %s"), *Package->GetPathName(), *Message);
-		}
-	}
-
-	FScriptExceptionHandler::Get().HandleException(ELogVerbosity::Warning, *Message, *FString());
-}
-
-void UAnimDataController::ReportError(const FText& InMessage) const
-{
-	FString Message = InMessage.ToString();
-	if (Model != nullptr)
-	{
-		if (UPackage* Package = Cast<UPackage>(Model->GetOutermost()))
-		{
-			Message = FString::Printf(TEXT("%s : %s"), *Package->GetPathName(), *Message);
-		}
-	}
-
-	FScriptExceptionHandler::Get().HandleException(ELogVerbosity::Error, *Message, *FString());
-}
-
-FString UAnimDataController::GetCurveTypeValueName(ERawCurveTrackTypes InType) const
-{
-	FString ValueString;
-
-	const UEnum* Enum = FindObject<UEnum>(nullptr, TEXT("/Script/Engine.ERawCurveTrackTypes"));
-	if (Enum)
-	{
-		ValueString = Enum->GetNameStringByValue((int64)InType);
-	}
-
-	return ValueString;
-}
-
-bool UAnimDataController::CheckOuterClass(UClass* InClass) const
-{
-	ValidateModel();
-	
-	const UObject* ModelOuter = Model->GetOuter();
-	if (ModelOuter)
-	{
-		const UClass* OuterClass = ModelOuter->GetClass();
-		if (OuterClass)
-		{
-			if (OuterClass == InClass || OuterClass->IsChildOf(InClass))
-			{
-				return true;
-			}
-			else
-			{
-				ReportErrorf(LOCTEXT("NoValidOuterClassError", "Incorrect outer object class found for Animation Data Model {0}, expected {1} actual {2}"), FText::FromString(Model->GetName()), FText::FromString(InClass->GetName()), FText::FromString(OuterClass->GetName()));
-			}
-		}
-	}
-	else
-	{
-		ReportErrorf(LOCTEXT("NoValidOuterObjectFoundError", "No valid outer object found for Animation Data Model {0}"), FText::FromString(Model->GetName()));
-	}
-
-	return false;
+	ResizeNumberOfFrames(ConvertSecondsToFrameNumber(Length), ConvertSecondsToFrameNumber(T0), ConvertSecondsToFrameNumber(T1), bShouldTransact);
 }
 
 int32 UAnimDataController::AddBoneTrack(FName BoneName, bool bShouldTransact /*= true*/)
 {
-	if (!CheckOuterClass(UAnimSequence::StaticClass()))
+	if (!ModelInterface->GetAnimationSequence())
 	{
 		return INDEX_NONE;
 	}
 
-	CONDITIONAL_TRANSACTION(LOCTEXT("AddBoneTrack", "Adding Animation Data Track"));
+	FTransaction Transaction = ConditionalTransaction(LOCTEXT("AddBoneTrack", "Adding Animation Data Track"), bShouldTransact);
 	return InsertBoneTrack(BoneName, INDEX_NONE, bShouldTransact);
 }
 
 int32 UAnimDataController::InsertBoneTrack(FName BoneName, int32 DesiredIndex, bool bShouldTransact /*= true*/)
 {
-	if (!CheckOuterClass(UAnimSequence::StaticClass()))
+	if (!ModelInterface->GetAnimationSequence())
 	{
 		return INDEX_NONE;
 	}
@@ -1331,29 +1260,26 @@ int32 UAnimDataController::InsertBoneTrack(FName BoneName, int32 DesiredIndex, b
 		if (Model->GetNumBoneTracks() >= MAX_ANIMATION_TRACKS)
 		{
 			ReportWarningf(LOCTEXT("MaxNumberOfTracksReachedWarning", "Cannot add track with name {0}. An animation sequence cannot contain more than 65535 tracks"), FText::FromName(BoneName));
+			return INDEX_NONE;
 		}
 		else
 		{
-			CONDITIONAL_TRANSACTION(LOCTEXT("InsertBoneTrack", "Inserting Animation Data Track"));
+			FTransaction Transaction = ConditionalTransaction(LOCTEXT("InsertBoneTrack", "Inserting Animation Data Track"), bShouldTransact);
 
 			// Determine correct index to do insertion at
 			const int32 InsertIndex = Model->BoneAnimationTracks.IsValidIndex(DesiredIndex) ? DesiredIndex : Model->BoneAnimationTracks.Num();
-
-			FBoneAnimationTrack& NewTrack = Model->BoneAnimationTracks.InsertDefaulted_GetRef(InsertIndex);
-			NewTrack.Name = BoneName;
+			int32 BoneIndex = INDEX_NONE;
 
 			if (const UAnimSequence* AnimationSequence = Model->GetAnimationSequence())
 			{
 				if (const USkeleton* Skeleton = AnimationSequence->GetSkeleton())
 				{
-					const int32 BoneIndex = Skeleton->GetReferenceSkeleton().FindBoneIndex(BoneName);
+					BoneIndex = Skeleton->GetReferenceSkeleton().FindBoneIndex(BoneName);
 
 					if (BoneIndex == INDEX_NONE)
 					{
-						ReportWarningf(LOCTEXT("UnableToFindBoneIndexWarning", "Unable to retrieve bone index for track: {0}"), FText::FromName(BoneName));
+						ReportErrorf(LOCTEXT("UnableToFindBoneIndexWarning", "Unable to retrieve bone index for track: {0}"), FText::FromName(BoneName));
 					}
-
-					NewTrack.BoneTreeIndex = BoneIndex;
 				}
 				else
 				{
@@ -1365,14 +1291,21 @@ int32 UAnimDataController::InsertBoneTrack(FName BoneName, int32 DesiredIndex, b
 				ReportError(LOCTEXT("UnableToGetOuterAnimSequenceError", "Unable to retrieve outer Animation Sequence"));
 			}
 
-			FAnimationTrackAddedPayload Payload;
-			Payload.Name = BoneName;
-			Payload.TrackIndex = InsertIndex;
+			if (BoneIndex != INDEX_NONE)
+			{
+				FBoneAnimationTrack& NewTrack = Model->BoneAnimationTracks.InsertDefaulted_GetRef(InsertIndex);
+				NewTrack.Name = BoneName;
+				NewTrack.BoneTreeIndex = BoneIndex;
+
+				FAnimationTrackAddedPayload Payload;
+				Payload.Name = BoneName;
+				Payload.TrackIndex = InsertIndex;
 
-			Model->Notify<FAnimationTrackAddedPayload>(EAnimDataModelNotifyType::TrackAdded, Payload);
-			CONDITIONAL_ACTION(UE::Anim::FRemoveTrackAction, NewTrack, InsertIndex);
+				Model->GetNotifier().Notify<FAnimationTrackAddedPayload>(EAnimDataModelNotifyType::TrackAdded, Payload);
+				ConditionalAction<UE::Anim::FRemoveTrackAction>(bShouldTransact, NewTrack);
 
-			return InsertIndex;
+				return InsertIndex;
+			}
 		}
 	}
 	else
@@ -1385,7 +1318,7 @@ int32 UAnimDataController::InsertBoneTrack(FName BoneName, int32 DesiredIndex, b
 
 bool UAnimDataController::RemoveBoneTrack(FName BoneName, bool bShouldTransact /*= true*/)
 {
-	if (!CheckOuterClass(UAnimSequence::StaticClass()))
+	if (!ModelInterface->GetAnimationSequence())
 	{
 		return false;
 	}
@@ -1394,7 +1327,7 @@ bool UAnimDataController::RemoveBoneTrack(FName BoneName, bool bShouldTransact /
 
 	if (ExistingTrackPtr != nullptr)
 	{
-		CONDITIONAL_TRANSACTION(LOCTEXT("RemoveBoneTrack", "Removing Animation Data Track"));
+		FTransaction Transaction = ConditionalTransaction(LOCTEXT("RemoveBoneTrack", "Removing Animation Data Track"), bShouldTransact);
 		const int32 TrackIndex = Model->BoneAnimationTracks.IndexOfByPredicate([ExistingTrackPtr](const FBoneAnimationTrack& Track)
 		{
 			return Track.Name == ExistingTrackPtr->Name;
@@ -1402,13 +1335,13 @@ bool UAnimDataController::RemoveBoneTrack(FName BoneName, bool bShouldTransact /
 
 		ensure(TrackIndex != INDEX_NONE);
 
-		CONDITIONAL_ACTION(UE::Anim::FAddTrackAction, *ExistingTrackPtr, TrackIndex);
+		ConditionalAction<UE::Anim::FAddTrackAction>(bShouldTransact,*ExistingTrackPtr);
 		Model->BoneAnimationTracks.RemoveAt(TrackIndex);
 
 		FAnimationTrackRemovedPayload Payload;
 		Payload.Name = BoneName;
 
-		Model->Notify(EAnimDataModelNotifyType::TrackRemoved, Payload);
+		Model->GetNotifier().Notify(EAnimDataModelNotifyType::TrackRemoved, Payload);
 
 		return true;
 	}
@@ -1422,7 +1355,7 @@ bool UAnimDataController::RemoveBoneTrack(FName BoneName, bool bShouldTransact /
 
 void UAnimDataController::RemoveAllBoneTracks(bool bShouldTransact /*= true*/)
 {
-	if (!CheckOuterClass(UAnimSequence::StaticClass()))
+	if (!ModelInterface->GetAnimationSequence())
 	{
 		return;
 	}
@@ -1432,7 +1365,7 @@ void UAnimDataController::RemoveAllBoneTracks(bool bShouldTransact /*= true*/)
 
 	if (TrackNames.Num())
 	{
-		CONDITIONAL_BRACKET(LOCTEXT("RemoveAllBoneTracks", "Removing all Animation Data Tracks"));
+		FBracket Bracket = ConditionalBracket(LOCTEXT("RemoveAllBoneTracks", "Removing all Animation Data Tracks"), bShouldTransact);
 		for (const FName& TrackName : TrackNames)
 		{
 			RemoveBoneTrack(TrackName, bShouldTransact);
@@ -1442,12 +1375,12 @@ void UAnimDataController::RemoveAllBoneTracks(bool bShouldTransact /*= true*/)
 
 bool UAnimDataController::SetBoneTrackKeys(FName BoneName, const TArray<FVector>& PositionalKeys, const TArray<FQuat>& RotationalKeys, const TArray<FVector>& ScalingKeys, bool bShouldTransact /*= true*/)
 {
-	if (!CheckOuterClass(UAnimSequence::StaticClass()))
+	if (!ModelInterface->GetAnimationSequence())
 	{
 		return false;
 	}
 
-	CONDITIONAL_TRANSACTION(LOCTEXT("SetTrackKeysTransaction", "Setting Animation Data Track keys"));
+	FTransaction Transaction = ConditionalTransaction(LOCTEXT("SetTrackKeysTransaction", "Setting Animation Data Track keys"), bShouldTransact);
 
 	// Validate key format
 	const int32 MaxNumKeys = FMath::Max(FMath::Max(PositionalKeys.Num(), RotationalKeys.Num()), ScalingKeys.Num());
@@ -1462,7 +1395,7 @@ bool UAnimDataController::SetBoneTrackKeys(FName BoneName, const TArray<FVector>
 		{
 			if (FBoneAnimationTrack* TrackPtr = Model->FindMutableBoneTrackByName(BoneName))
 			{
-				CONDITIONAL_ACTION(UE::Anim::FSetTrackKeysAction, *TrackPtr);
+				ConditionalAction<UE::Anim::FSetTrackKeysAction>(bShouldTransact, *TrackPtr);
 
 
 				TrackPtr->InternalTrackData.PosKeys.SetNum(MaxNumKeys);
@@ -1478,7 +1411,7 @@ bool UAnimDataController::SetBoneTrackKeys(FName BoneName, const TArray<FVector>
 				FAnimationTrackChangedPayload Payload;
 				Payload.Name = BoneName;
 
-				Model->Notify(EAnimDataModelNotifyType::TrackChanged, Payload);
+				Model->GetNotifier().Notify(EAnimDataModelNotifyType::TrackChanged, Payload);
 
 				return true;
 			}
@@ -1507,7 +1440,7 @@ bool UAnimDataController::SetBoneTrackKeys(FName BoneName, const TArray<FVector3
 		return false;
 	}
 
-	CONDITIONAL_TRANSACTION(LOCTEXT("SetTrackKeysTransaction", "Setting Animation Data Track keys"));
+	FTransaction Transaction = ConditionalTransaction(LOCTEXT("SetTrackKeysTransaction", "Setting Animation Data Track keys"), bShouldTransact);
 
 	// Validate key format
 	const int32 MaxNumKeys = FMath::Max(FMath::Max(PositionalKeys.Num(), RotationalKeys.Num()), ScalingKeys.Num());
@@ -1522,7 +1455,7 @@ bool UAnimDataController::SetBoneTrackKeys(FName BoneName, const TArray<FVector3
 		{
 			if (FBoneAnimationTrack* TrackPtr = Model->FindMutableBoneTrackByName(BoneName))
 			{
-				CONDITIONAL_ACTION(UE::Anim::FSetTrackKeysAction, *TrackPtr);
+				ConditionalAction<UE::Anim::FSetTrackKeysAction>(bShouldTransact, *TrackPtr);
 
 				TrackPtr->InternalTrackData.PosKeys = PositionalKeys;
 				TrackPtr->InternalTrackData.RotKeys = RotationalKeys;
@@ -1531,7 +1464,7 @@ bool UAnimDataController::SetBoneTrackKeys(FName BoneName, const TArray<FVector3
 				FAnimationTrackChangedPayload Payload;
 				Payload.Name = BoneName;
 
-				Model->Notify(EAnimDataModelNotifyType::TrackChanged, Payload);
+				Model->GetNotifier().Notify(EAnimDataModelNotifyType::TrackChanged, Payload);
 
 				return true;
 			}
@@ -1553,23 +1486,23 @@ bool UAnimDataController::SetBoneTrackKeys(FName BoneName, const TArray<FVector3
 	return false;
 }
 
-int32 DiscreteInclusiveLower(const TRange<int32>& InRange)
+static int32 DiscreteInclusiveLower(const TRange<int32>& InRange)
 {
 	check(!InRange.GetLowerBound().IsOpen());
 
 	// Add one for exclusive lower bounds since they start on the next subsequent frame
-	static const int32 Offsets[]   = { 0, 1 };
+	static constexpr int32 Offsets[]   = { 0, 1 };
 	const int32        OffsetIndex = (int32)InRange.GetLowerBound().IsExclusive();
 
 	return InRange.GetLowerBound().GetValue() + Offsets[OffsetIndex];
 }
 
-int32 DiscreteExclusiveUpper(const TRange<int32>& InRange)
+static int32 DiscreteExclusiveUpper(const TRange<int32>& InRange)
 {
 	check(!InRange.GetUpperBound().IsOpen());
 
 	// Add one for inclusive upper bounds since they finish on the next subsequent frame
-	static const int32 Offsets[]   = { 0, 1 };
+	static constexpr int32 Offsets[]   = { 0, 1 };
 	const int32        OffsetIndex = (int32)InRange.GetUpperBound().IsInclusive();
 
 	return InRange.GetUpperBound().GetValue() + Offsets[OffsetIndex];
@@ -1582,7 +1515,7 @@ bool UAnimDataController::UpdateBoneTrackKeys(FName BoneName, const FInt32Range&
 		return false;
 	}
 
-	CONDITIONAL_TRANSACTION(LOCTEXT("SetTrackKeysRangeTransaction", "Setting Animation Data Track keys"));
+	FTransaction Transaction = ConditionalTransaction(LOCTEXT("SetTrackKeysRangeTransaction", "Setting Animation Data Track keys"), bShouldTransact);
 
 	// Validate key format
 	const int32 MaxNumKeys = FMath::Max(FMath::Max(PositionalKeys.Num(), RotationalKeys.Num()), ScalingKeys.Num());
@@ -1609,7 +1542,7 @@ bool UAnimDataController::UpdateBoneTrackKeys(FName BoneName, const FInt32Range&
 						TArray<FQuat4f>& TrackRotKeys = InternalTrackData.RotKeys;
 						TArray<FVector3f>& TrackScaleKeys = InternalTrackData.ScaleKeys;
 
-						CONDITIONAL_ACTION(UE::Anim::FSetTrackKeysAction, *TrackPtr);
+						ConditionalAction<UE::Anim::FSetTrackKeysAction>(bShouldTransact, *TrackPtr);
 
 						int32 KeyIndex = 0;
 						for (int32 FrameIndex = RangeMin; FrameIndex < RangeMax; ++FrameIndex, ++KeyIndex)
@@ -1622,7 +1555,7 @@ bool UAnimDataController::UpdateBoneTrackKeys(FName BoneName, const FInt32Range&
 						FAnimationTrackChangedPayload Payload;
 						Payload.Name = BoneName;
 
-						Model->Notify(EAnimDataModelNotifyType::TrackChanged, Payload);
+						Model->GetNotifier().Notify(EAnimDataModelNotifyType::TrackChanged, Payload);
 
 						return true;
 					}
@@ -1665,7 +1598,7 @@ bool UAnimDataController::UpdateBoneTrackKeys(FName BoneName, const FInt32Range&
 		return false;
 	}
 
-	CONDITIONAL_TRANSACTION(LOCTEXT("SetTrackKeysTransaction", "Setting Animation Data Track keys"));
+	FTransaction Transaction = ConditionalTransaction(LOCTEXT("SetTrackKeysTransaction", "Setting Animation Data Track keys"), bShouldTransact);
 
 	// Validate key format
 	const int32 MaxNumKeys = FMath::Max(FMath::Max(PositionalKeys.Num(), RotationalKeys.Num()), ScalingKeys.Num());
@@ -1692,7 +1625,7 @@ bool UAnimDataController::UpdateBoneTrackKeys(FName BoneName, const FInt32Range&
 						TArray<FQuat4f>& TrackRotKeys = InternalTrackData.RotKeys;
 						TArray<FVector3f>& TrackScaleKeys = InternalTrackData.ScaleKeys;
 
-						CONDITIONAL_ACTION(UE::Anim::FSetTrackKeysAction, *TrackPtr);
+						ConditionalAction<UE::Anim::FSetTrackKeysAction>(bShouldTransact, *TrackPtr);
 
 						int32 KeyIndex = 0;
 						for (int32 FrameIndex = RangeMin; FrameIndex < RangeMax; ++FrameIndex, ++KeyIndex)
@@ -1705,7 +1638,7 @@ bool UAnimDataController::UpdateBoneTrackKeys(FName BoneName, const FInt32Range&
 						FAnimationTrackChangedPayload Payload;
 						Payload.Name = BoneName;
 
-						Model->Notify(EAnimDataModelNotifyType::TrackChanged, Payload);
+						Model->GetNotifier().Notify(EAnimDataModelNotifyType::TrackChanged, Payload);
 
 						return true;
 					}
@@ -1743,7 +1676,7 @@ bool UAnimDataController::UpdateBoneTrackKeys(FName BoneName, const FInt32Range&
 
 void UAnimDataController::ResizeCurves(float NewLength, bool bInserted, float T0, float T1, bool bShouldTransact /*= true*/)
 {
-	CONDITIONAL_BRACKET(LOCTEXT("ResizeCurves", "Resizing all Curves"));
+	FBracket Bracket = ConditionalBracket(LOCTEXT("ResizeCurves", "Resizing all Curves"), bShouldTransact);
 
 	for (FFloatCurve& Curve : Model->CurveData.FloatCurves)
 	{
@@ -1757,11 +1690,11 @@ void UAnimDataController::ResizeCurves(float NewLength, bool bInserted, float T0
 		FTransformCurve ResizedCurve = Curve;
 		for (int32 SubCurveIndex = 0; SubCurveIndex < 3; ++SubCurveIndex)
 		{
-			const ETransformCurveChannel Channel = (ETransformCurveChannel)SubCurveIndex;
+			const ETransformCurveChannel Channel = static_cast<ETransformCurveChannel>(SubCurveIndex);
 			FVectorCurve& SubCurve = *ResizedCurve.GetVectorCurveByIndex(SubCurveIndex);
 			for (int32 ChannelIndex = 0; ChannelIndex < 3; ++ChannelIndex)
 			{
-				const EVectorCurveChannel Axis = (EVectorCurveChannel)ChannelIndex;
+				const EVectorCurveChannel Axis = static_cast<EVectorCurveChannel>(ChannelIndex);
 				FAnimationCurveIdentifier TargetCurveIdentifier = FAnimationCurveIdentifier(Curve.Name, ERawCurveTrackTypes::RCT_Transform);
 				UAnimationCurveIdentifierExtensions::GetTransformChildCurveIdentifier(TargetCurveIdentifier, Channel, Axis);
 				
@@ -1775,7 +1708,7 @@ void UAnimDataController::ResizeCurves(float NewLength, bool bInserted, float T0
 
 void UAnimDataController::ResizeAttributes(float NewLength, bool bInserted, float T0, float T1, bool bShouldTransact)
 {
-	CONDITIONAL_BRACKET(LOCTEXT("ResizeAttributes", "Resizing all Attributes"));
+	FBracket Bracket = ConditionalBracket(LOCTEXT("ResizeAttributes", "Resizing all Attributes"), bShouldTransact);
 
 	for (FAnimatedBoneAttribute& Attribute : Model->AnimatedBoneAttributes)
 	{
@@ -1807,18 +1740,18 @@ bool UAnimDataController::AddAttribute(const FAnimationAttributeIdentifier& Attr
 
 		if (!bAttributeAlreadyExists)
 		{
-			CONDITIONAL_TRANSACTION(LOCTEXT("AddAttribute", "Adding Animated Bone Attribute"));
+			FTransaction Transaction = ConditionalTransaction(LOCTEXT("AddAttribute", "Adding Animated Bone Attribute"), bShouldTransact);
 
 			FAnimatedBoneAttribute& Attribute = Model->AnimatedBoneAttributes.AddDefaulted_GetRef();
 			Attribute.Identifier = AttributeIdentifier;
 
 			Attribute.Curve.SetScriptStruct(AttributeIdentifier.GetType());
 		
-			CONDITIONAL_ACTION(UE::Anim::FRemoveAtributeAction, AttributeIdentifier);
+			ConditionalAction<UE::Anim::FRemoveAtributeAction>(bShouldTransact, AttributeIdentifier);
 
 			FAttributeAddedPayload Payload;
 			Payload.Identifier = AttributeIdentifier;
-			Model->Notify(EAnimDataModelNotifyType::AttributeAdded, Payload);
+			Model->GetNotifier().Notify(EAnimDataModelNotifyType::AttributeAdded, Payload);
 
 			return true;
 		}
@@ -1847,15 +1780,15 @@ bool UAnimDataController::RemoveAttribute(const FAnimationAttributeIdentifier& A
 
 		if (AttributeIndex != INDEX_NONE)
 		{
-			CONDITIONAL_TRANSACTION(LOCTEXT("RemoveAttribute", "Removing Animated Bone Attribute"));
+			FTransaction Transaction = ConditionalTransaction(LOCTEXT("RemoveAttribute", "Removing Animated Bone Attribute"), bShouldTransact);
 
-			CONDITIONAL_ACTION(UE::Anim::FAddAtributeAction, Model->AnimatedBoneAttributes[AttributeIndex]);
+			ConditionalAction<UE::Anim::FAddAtributeAction>(bShouldTransact, Model->AnimatedBoneAttributes[AttributeIndex]);
 
 			Model->AnimatedBoneAttributes.RemoveAtSwap(AttributeIndex);
 			
 			FAttributeRemovedPayload Payload;
 			Payload.Identifier = AttributeIdentifier;
-			Model->Notify(EAnimDataModelNotifyType::AttributeRemoved, Payload);
+			Model->GetNotifier().Notify(EAnimDataModelNotifyType::AttributeRemoved, Payload);
 
 			return true;
 		}
@@ -1892,7 +1825,7 @@ int32 UAnimDataController::RemoveAllAttributesForBone(const FName& BoneName, boo
 
 	if (Identifiers.Num())
 	{
-		CONDITIONAL_BRACKET(LOCTEXT("RemoveAllAttributesForBone", "Removing all Attributes for Bone"));
+		FBracket Bracket = ConditionalBracket(LOCTEXT("RemoveAllAttributesForBone", "Removing all Attributes for Bone"), bShouldTransact);
 		for (const FAnimationAttributeIdentifier& Identifier : Identifiers)
 		{
 			NumRemovedAttributes += RemoveAttribute(Identifier, bShouldTransact) ? 1 : 0;
@@ -1915,7 +1848,7 @@ int32 UAnimDataController::RemoveAllAttributes(bool bShouldTransact)
 
 	if (Identifiers.Num())
 	{
-		CONDITIONAL_BRACKET(LOCTEXT("RemoveAllAttributes", "Removing all Attributes"));
+		FBracket Bracket = ConditionalBracket(LOCTEXT("RemoveAllAttributes", "Removing all Attributes"), bShouldTransact);
 		for (const FAnimationAttributeIdentifier& Identifier : Identifiers)
 		{
 			NumRemovedAttributes += RemoveAttribute(Identifier, bShouldTransact) ? 1: 0;
@@ -1931,7 +1864,7 @@ bool UAnimDataController::SetAttributeKey_Internal(const FAnimationAttributeIden
 	{
 		if (KeyValue)
 		{
-			FAnimatedBoneAttribute* AttributePtr = Model->AnimatedBoneAttributes.FindByPredicate([AttributeIdentifier](FAnimatedBoneAttribute& Attribute)
+			FAnimatedBoneAttribute* AttributePtr = Model->AnimatedBoneAttributes.FindByPredicate([AttributeIdentifier](const FAnimatedBoneAttribute& Attribute)
 			{
 				return Attribute.Identifier == AttributeIdentifier;
 			});
@@ -1940,26 +1873,26 @@ bool UAnimDataController::SetAttributeKey_Internal(const FAnimationAttributeIden
 			{
 				if (TypeStruct == AttributePtr->Identifier.GetType())
 				{
-					CONDITIONAL_TRANSACTION(LOCTEXT("SettingAttributeKey", "Setting Animated Bone Attribute key"));
+					FTransaction Transaction = ConditionalTransaction(LOCTEXT("SettingAttributeKey", "Setting Animated Bone Attribute key"), bShouldTransact);
 
 					FAttributeCurve& Curve = AttributePtr->Curve;
-					FKeyHandle KeyHandle = Curve.FindKey(Time);
+					const FKeyHandle KeyHandle = Curve.FindKey(Time);
 					// In case the key does not yet exist one will be added, and thus the undo is a remove
 					if (KeyHandle == FKeyHandle::Invalid())
 					{
-						CONDITIONAL_ACTION(UE::Anim::FRemoveAtributeKeyAction, AttributeIdentifier, Time);
-						Curve.UpdateOrAddKey(Time, KeyValue);
+						ConditionalAction<UE::Anim::FRemoveAtributeKeyAction>(bShouldTransact, AttributeIdentifier, Time);
+						Curve.UpdateOrAddTypedKey(Time, KeyValue, TypeStruct);
 					}
 					// In case the key does exist it will be updated , and thus the undo is a revert to the current value
 					else
 					{
-						CONDITIONAL_ACTION(UE::Anim::FSetAtributeKeyAction, AttributeIdentifier, Curve.GetKey(KeyHandle));
-						Curve.UpdateOrAddKey(Time, KeyValue);
+						ConditionalAction<UE::Anim::FSetAtributeKeyAction>(bShouldTransact, AttributeIdentifier, Curve.GetKey(KeyHandle));
+						Curve.UpdateOrAddTypedKey(Time, KeyValue, TypeStruct);
 					}
 
 					FAttributeChangedPayload Payload;
 					Payload.Identifier = AttributeIdentifier;
-					Model->Notify(EAnimDataModelNotifyType::AttributeChanged, Payload);
+					Model->GetNotifier().Notify(EAnimDataModelNotifyType::AttributeChanged, Payload);
 
 					return true;
 				}
@@ -1994,7 +1927,7 @@ bool UAnimDataController::SetAttributeKeys_Internal(const FAnimationAttributeIde
 	{
 		if (Times.Num() == KeyValues.Num())
 		{
-			FAnimatedBoneAttribute* AttributePtr = Model->AnimatedBoneAttributes.FindByPredicate([AttributeIdentifier](FAnimatedBoneAttribute& Attribute)
+			FAnimatedBoneAttribute* AttributePtr = Model->AnimatedBoneAttributes.FindByPredicate([AttributeIdentifier](const FAnimatedBoneAttribute& Attribute)
 			{
 				return Attribute.Identifier == AttributeIdentifier;
 			});
@@ -2003,17 +1936,17 @@ bool UAnimDataController::SetAttributeKeys_Internal(const FAnimationAttributeIde
 			{
 				if (TypeStruct == AttributePtr->Identifier.GetType())
 				{
-					CONDITIONAL_TRANSACTION(LOCTEXT("SettingAttributeKeys", "Setting Animated Bone Attribute keys"));
+					FTransaction Transaction = ConditionalTransaction(LOCTEXT("SettingAttributeKeys", "Setting Animated Bone Attribute keys"), bShouldTransact);
 
 					FAnimatedBoneAttribute& Attribute = *AttributePtr;
 
-					CONDITIONAL_ACTION(UE::Anim::FSetAtributeKeysAction, Attribute);
+					ConditionalAction<UE::Anim::FSetAtributeKeysAction>(bShouldTransact, Attribute);
 			
 					Attribute.Curve.SetKeys(Times, KeyValues);
 
 					FAttributeChangedPayload Payload;
 					Payload.Identifier = AttributeIdentifier;
-					Model->Notify(EAnimDataModelNotifyType::AttributeChanged, Payload);
+					Model->GetNotifier().Notify(EAnimDataModelNotifyType::AttributeChanged, Payload);
 
 					return true;
 				}
@@ -2056,19 +1989,19 @@ bool UAnimDataController::RemoveAttributeKey(const FAnimationAttributeIdentifier
 		if (AttributePtr)
 		{
 			FAttributeCurve& Curve = AttributePtr->Curve;
-			FKeyHandle KeyHandle = Curve.FindKey(Time);
+			const FKeyHandle KeyHandle = Curve.FindKey(Time);
 
 			if (KeyHandle != FKeyHandle::Invalid())
 			{
-				CONDITIONAL_TRANSACTION(LOCTEXT("RemovingAttributeKey", "Removing Animated Bone Attribute key"));
+				FTransaction Transaction = ConditionalTransaction(LOCTEXT("RemovingAttributeKey", "Removing Animated Bone Attribute key"), bShouldTransact);
 
-				CONDITIONAL_ACTION(UE::Anim::FAddAtributeKeyAction, AttributeIdentifier, Curve.GetKey(KeyHandle));
+				ConditionalAction<UE::Anim::FAddAtributeKeyAction>(bShouldTransact, AttributeIdentifier, Curve.GetKey(KeyHandle));
 
 				Curve.DeleteKey(KeyHandle);
 
 				FAttributeAddedPayload Payload;
 				Payload.Identifier = AttributeIdentifier;
-				Model->Notify(EAnimDataModelNotifyType::AttributeChanged, Payload);
+				Model->GetNotifier().Notify(EAnimDataModelNotifyType::AttributeChanged, Payload);
 
 				return true;
 			}
@@ -2104,7 +2037,7 @@ bool UAnimDataController::DuplicateAttribute(const FAnimationAttributeIdentifier
 			{
 				if(const FAnimatedBoneAttribute* AttributePtr = Model->FindAttribute(AttributeIdentifier))
 				{
-					CONDITIONAL_TRANSACTION(LOCTEXT("DuplicateAttribute", "Duplicating Animation Attribute"));
+					FTransaction Transaction = ConditionalTransaction(LOCTEXT("DuplicateAttribute", "Duplicating Animation Attribute"), bShouldTransact);
 
 					FAnimatedBoneAttribute& DuplicateAttribute = Model->AnimatedBoneAttributes.AddDefaulted_GetRef();
 					DuplicateAttribute.Identifier = NewAttributeIdentifier;
@@ -2112,9 +2045,9 @@ bool UAnimDataController::DuplicateAttribute(const FAnimationAttributeIdentifier
 
 					FAttributeAddedPayload Payload;
 					Payload.Identifier = NewAttributeIdentifier;
-					Model->Notify(EAnimDataModelNotifyType::AttributeAdded, Payload);
+					Model->GetNotifier().Notify(EAnimDataModelNotifyType::AttributeAdded, Payload);
 
-					CONDITIONAL_ACTION(UE::Anim::FRemoveAtributeAction, NewAttributeIdentifier);
+					ConditionalAction<UE::Anim::FRemoveAtributeAction>(bShouldTransact, NewAttributeIdentifier);
 					
 					return true;
 				}
@@ -2141,6 +2074,21 @@ bool UAnimDataController::DuplicateAttribute(const FAnimationAttributeIdentifier
 	return false;
 }
 
+void UAnimDataController::UpdateWithSkeleton(USkeleton* TargetSkeleton, bool bShouldTransact)
+{
+	RemoveBoneTracksMissingFromSkeleton(TargetSkeleton);
+}
+
+void UAnimDataController::PopulateWithExistingModel(TScriptInterface<IAnimationDataModel> InModel)
+{	
+	Model->BoneAnimationTracks = InModel->GetBoneAnimationTracks();
+	Model->FrameRate = InModel->GetFrameRate();
+	Model->NumberOfFrames = InModel->GetNumberOfFrames();
+	Model->NumberOfKeys = InModel->GetNumberOfFrames() + 1;
+	Model->CurveData = InModel->GetCurveData();
+	Model->AnimatedBoneAttributes = InModel->GetAttributes();
+}
+
 #endif // WITH_EDITOR
 
 #undef LOCTEXT_NAMESPACE // "AnimDataController"
diff --git a/Engine/Source/Developer/AnimationDataController/Private/AnimDataControllerActions.cpp b/Engine/Source/Developer/AnimationDataController/Private/AnimDataControllerActions.cpp
--- a/Engine/Source/Developer/AnimationDataController/Private/AnimDataControllerActions.cpp
+++ b/Engine/Source/Developer/AnimationDataController/Private/AnimDataControllerActions.cpp
@@ -3,7 +3,7 @@
 #include "AnimDataControllerActions.h"
 #include "AnimDataController.h"
 
-#include "Animation/AnimData/AnimDataModel.h"
+#include "Animation/AnimData/IAnimationDataModel.h"
 #include "UObject/StrongObjectPtr.h"
 
 #define LOCTEXT_NAMESPACE "AnimDataControllerActions"
@@ -16,13 +16,13 @@ namespace Anim {
 
 TUniquePtr<FChange> FAnimDataBaseAction::Execute(UObject* Object)
 {
-	UAnimDataModel* Model = Cast<UAnimDataModel>(Object);
-	checkf(Model != nullptr, TEXT("Invalid UAnimDataModel Object"));
+	const TScriptInterface<IAnimationDataModel> DataModelInterface(Object);
+	checkf(DataModelInterface, TEXT("Invalid IAnimationDataModel Object"));
 
-	TStrongObjectPtr<UAnimDataController> Controller(NewObject<UAnimDataController>());
-	Controller->SetModel(Model);
+	const TScriptInterface<IAnimationDataController> Controller = DataModelInterface->GetController();
+	Controller->SetModel(DataModelInterface);
 
-	return ExecuteInternal(Model, Controller.Get());
+	return ExecuteInternal(DataModelInterface.GetInterface(), Controller.GetInterface());
 }
 
 FString FAnimDataBaseAction::ToString() const
@@ -30,7 +30,7 @@ FString FAnimDataBaseAction::ToString() const
 	return ToStringInternal();
 }
 
-TUniquePtr<FChange> FOpenBracketAction::ExecuteInternal(UAnimDataModel* Model, UAnimDataController* Controller)
+TUniquePtr<FChange> FOpenBracketAction::ExecuteInternal(IAnimationDataModel* Model, IAnimationDataController* Controller)
 {
 	Controller->NotifyBracketOpen();
 	return MakeUnique<FCloseBracketAction>(Description);
@@ -41,7 +41,7 @@ FString FOpenBracketAction::ToStringInternal() const
 	return FString::Printf(TEXT("Open Bracket: %s"), *Description);
 }
 
-TUniquePtr<FChange> FCloseBracketAction::ExecuteInternal(UAnimDataModel* Model, UAnimDataController* Controller)
+TUniquePtr<FChange> FCloseBracketAction::ExecuteInternal(IAnimationDataModel* Model, IAnimationDataController* Controller)
 {
 	Controller->NotifyBracketClosed();
 	return MakeUnique<FOpenBracketAction>(Description);
@@ -52,17 +52,16 @@ FString FCloseBracketAction::ToStringInternal() const
 	return TEXT("Closing Bracket");
 }
 
-FAddTrackAction::FAddTrackAction(const FBoneAnimationTrack& Track, int32 InTrackIndex)
+FAddTrackAction::FAddTrackAction(const FBoneAnimationTrack& Track)
 {
 	Name = Track.Name;
 	BoneTreeIndex = Track.BoneTreeIndex;
 	Data = Track.InternalTrackData;
-	TrackIndex = InTrackIndex;
 }
 
-TUniquePtr<FChange> FAddTrackAction::ExecuteInternal(UAnimDataModel* Model, UAnimDataController* Controller)
+TUniquePtr<FChange> FAddTrackAction::ExecuteInternal(IAnimationDataModel* Model, IAnimationDataController* Controller)
 {
-	Controller->InsertBoneTrack(Name, TrackIndex, false);
+	Controller->AddBoneTrack(Name, false);
 	
 	if (Data.PosKeys.Num() || Data.RotKeys.Num() || Data.ScaleKeys.Num())
 	{
@@ -70,7 +69,7 @@ TUniquePtr<FChange> FAddTrackAction::ExecuteInternal(UAnimDataModel* Model, UAni
 	}
 	
 	const FBoneAnimationTrack& AddedTrack = Model->GetBoneTrackByName(Name);
-	return MakeUnique<FRemoveTrackAction>(AddedTrack, TrackIndex);
+	return MakeUnique<FRemoveTrackAction>(AddedTrack);
 }
 
 FString FAddTrackAction::ToStringInternal() const
@@ -78,17 +77,16 @@ FString FAddTrackAction::ToStringInternal() const
 	return FText::Format(LOCTEXT("AddTrackAction_Description", "Adding animation bone track '{0}'."), FText::FromName(Name)).ToString();
 }
 
-FRemoveTrackAction::FRemoveTrackAction(const FBoneAnimationTrack& Track, int32 InTrackIndex)
+FRemoveTrackAction::FRemoveTrackAction(const FBoneAnimationTrack& Track)
 {
 	Name = Track.Name;
-	TrackIndex = InTrackIndex;
 }
 
-TUniquePtr<FChange> FRemoveTrackAction::ExecuteInternal(UAnimDataModel* Model, UAnimDataController* Controller)
+TUniquePtr<FChange> FRemoveTrackAction::ExecuteInternal(IAnimationDataModel* Model, IAnimationDataController* Controller)
 {
 	const FBoneAnimationTrack& Track = Model->GetBoneTrackByName(Name);
 
-	TUniquePtr<FChange> InverseAction = MakeUnique<FAddTrackAction>(Track, TrackIndex);
+	TUniquePtr<FChange> InverseAction = MakeUnique<FAddTrackAction>(Track);
 	Controller->RemoveBoneTrack(Name, false);
 	return InverseAction;
 }
@@ -105,7 +103,7 @@ FSetTrackKeysAction::FSetTrackKeysAction(const FBoneAnimationTrack& Track)
 	TrackData = Track.InternalTrackData;
 }
 
-TUniquePtr<FChange> FSetTrackKeysAction::ExecuteInternal(UAnimDataModel* Model, UAnimDataController* Controller)
+TUniquePtr<FChange> FSetTrackKeysAction::ExecuteInternal(IAnimationDataModel* Model, IAnimationDataController* Controller)
 {
 	const FBoneAnimationTrack& Track = Model->GetBoneTrackByName(Name);
 	TUniquePtr<FChange> InverseAction = MakeUnique<FSetTrackKeysAction>(Track);
@@ -119,29 +117,29 @@ FString FSetTrackKeysAction::ToStringInternal() const
 	return FText::Format(LOCTEXT("SetTrackKeysAction_Description", "Setting keys for animation bone track '{0}'."), FText::FromName(Name)).ToString();
 }
 
-FResizePlayLengthAction::FResizePlayLengthAction(const UAnimDataModel* InModel, float t0, float t1) : T0(t0), T1(t1)
+FResizePlayLengthInFramesAction::FResizePlayLengthInFramesAction(const IAnimationDataModel* InModel, FFrameNumber F0, FFrameNumber F1) : Frame0(F0), Frame1(F1)
 {
-	Length = InModel->GetPlayLength();
+	Length = InModel->GetNumberOfFrames();
 }
 
-TUniquePtr<FChange> FResizePlayLengthAction::ExecuteInternal(UAnimDataModel* Model, UAnimDataController* Controller)
+TUniquePtr<FChange> FResizePlayLengthInFramesAction::ExecuteInternal(IAnimationDataModel* Model, IAnimationDataController* Controller)
 {
-	TUniquePtr<FChange> InverseAction = MakeUnique<FResizePlayLengthAction>(Model, T0, T1);
-	Controller->ResizePlayLength(Length, T0, T1, false);
+	TUniquePtr<FChange> InverseAction = MakeUnique<FResizePlayLengthInFramesAction>(Model, Frame0, Frame1);
+	Controller->ResizeNumberOfFrames(Length, Frame0, Frame1, false);
 	return InverseAction;
 }
 
-FString FResizePlayLengthAction::ToStringInternal() const
+FString FResizePlayLengthInFramesAction::ToStringInternal() const
 {
-	return FText::Format(LOCTEXT("ResizePlayLengthAction_Description", "Resizing play length to {0}."), FText::AsNumber(Length)).ToString();
+	return FText::Format(LOCTEXT("ResizePlayLengthInFramesAction_Description", "Resizing play length to {0} frames."), FText::AsNumber(Length.Value)).ToString();
 }
 
-FSetFrameRateAction::FSetFrameRateAction(const UAnimDataModel* InModel)
+FSetFrameRateAction::FSetFrameRateAction(const IAnimationDataModel* InModel)
 {
 	FrameRate = InModel->GetFrameRate();
 }
 
-TUniquePtr<FChange> FSetFrameRateAction::ExecuteInternal(UAnimDataModel* Model, UAnimDataController* Controller)
+TUniquePtr<FChange> FSetFrameRateAction::ExecuteInternal(IAnimationDataModel* Model, IAnimationDataController* Controller)
 {
 	TUniquePtr<FChange> InverseAction = MakeUnique<FSetFrameRateAction>(Model);
 	Controller->SetFrameRate(FrameRate, false);
@@ -153,7 +151,7 @@ FString FSetFrameRateAction::ToStringInternal() const
 	return FText::Format(LOCTEXT("SetFrameRateAction_Description", "Setting Frame Rate to {0}."), FrameRate.ToPrettyText()).ToString();
 }
 
-TUniquePtr<FChange> FAddCurveAction::ExecuteInternal(UAnimDataModel* Model, UAnimDataController* Controller)
+TUniquePtr<FChange> FAddCurveAction::ExecuteInternal(IAnimationDataModel* Model, IAnimationDataController* Controller)
 {
 	TUniquePtr<FChange> InverseAction = MakeUnique<FRemoveCurveAction>(CurveId);
 
@@ -170,7 +168,7 @@ FString FAddCurveAction::ToStringInternal() const
 	return FText::Format(LOCTEXT("AddCurveAction_Description", "Adding {0} curve '{1}'."), FText::FromString(CurveId.CurveType == ERawCurveTrackTypes::RCT_Float ? FloatLabel : TransformLabel), FText::FromName(CurveId.InternalName.DisplayName)).ToString();
 }
 
-TUniquePtr<FChange> FRemoveCurveAction::ExecuteInternal(UAnimDataModel* Model, UAnimDataController* Controller)
+TUniquePtr<FChange> FRemoveCurveAction::ExecuteInternal(IAnimationDataModel* Model, IAnimationDataController* Controller)
 {
 	TUniquePtr<FChange> InverseAction;	
 
@@ -198,7 +196,7 @@ FString FRemoveCurveAction::ToStringInternal() const
 	return FText::Format(LOCTEXT("RemoveCurveAction_Description", "Removing {0} curve '{1}'."), FText::FromString(CurveId.CurveType == ERawCurveTrackTypes::RCT_Float ? FloatLabel : TransformLabel), FText::FromName(CurveId.InternalName.DisplayName)).ToString();
 }
 
-TUniquePtr<FChange> FSetCurveFlagsAction::ExecuteInternal(UAnimDataModel* Model, UAnimDataController* Controller)
+TUniquePtr<FChange> FSetCurveFlagsAction::ExecuteInternal(IAnimationDataModel* Model, IAnimationDataController* Controller)
 {
 	const FAnimCurveBase& Curve = Model->GetCurve(CurveId);
 
@@ -218,7 +216,7 @@ FString FSetCurveFlagsAction::ToStringInternal() const
 	return FText::Format(LOCTEXT("SetCurveFlagsAction_Description", "Setting flags for {0} curve '{1}'."), FText::FromString(CurveType == ERawCurveTrackTypes::RCT_Float ? FloatLabel : TransformLabel), FText::FromName(CurveId.InternalName.DisplayName)).ToString();
 }
 
-TUniquePtr<FChange> FRenameCurveAction::ExecuteInternal(UAnimDataModel* Model, UAnimDataController* Controller)
+TUniquePtr<FChange> FRenameCurveAction::ExecuteInternal(IAnimationDataModel* Model, IAnimationDataController* Controller)
 {
 	Controller->RenameCurve(CurveId, NewCurveId, false);
 	return MakeUnique<FRenameCurveAction>(NewCurveId, CurveId);
@@ -236,7 +234,7 @@ FScaleCurveAction::FScaleCurveAction(const FAnimationCurveIdentifier& InCurveId,
 	ensure(CurveType == ERawCurveTrackTypes::RCT_Float);
 }
 
-TUniquePtr<FChange> FScaleCurveAction::ExecuteInternal(UAnimDataModel* Model, UAnimDataController* Controller)
+TUniquePtr<FChange> FScaleCurveAction::ExecuteInternal(IAnimationDataModel* Model, IAnimationDataController* Controller)
 {
 	const float InverseFactor = 1.0f / Factor;
 
@@ -255,7 +253,7 @@ FString FScaleCurveAction::ToStringInternal() const
 	return FText::Format(LOCTEXT("ScaleCurveAction_Description", "Scaling {0} curve '{1}'."), FText::FromString(CurveType == ERawCurveTrackTypes::RCT_Float ? FloatLabel : TransformLabel), FText::FromName(CurveId.InternalName.DisplayName)).ToString();
 }
 
-TUniquePtr<FChange> FAddFloatCurveAction::ExecuteInternal(UAnimDataModel* Model, UAnimDataController* Controller)
+TUniquePtr<FChange> FAddFloatCurveAction::ExecuteInternal(IAnimationDataModel* Model, IAnimationDataController* Controller)
 {
 	TUniquePtr<FChange> InverseAction = MakeUnique<FRemoveCurveAction>(CurveId);
 
@@ -288,7 +286,7 @@ FAddTransformCurveAction::FAddTransformCurveAction(const FAnimationCurveIdentifi
 	SubCurveKeys[8] = InTransformCurve.ScaleCurve.FloatCurves[2].GetConstRefOfKeys();
 }
 
-TUniquePtr<FChange> FAddTransformCurveAction::ExecuteInternal(UAnimDataModel* Model, UAnimDataController* Controller)
+TUniquePtr<FChange> FAddTransformCurveAction::ExecuteInternal(IAnimationDataModel* Model, IAnimationDataController* Controller)
 {
 	TUniquePtr<FChange> InverseAction = MakeUnique<FRemoveCurveAction>(CurveId);
 
@@ -296,10 +294,10 @@ TUniquePtr<FChange> FAddTransformCurveAction::ExecuteInternal(UAnimDataModel* Mo
 
 	for (int32 SubCurveIndex = 0; SubCurveIndex < 3; ++SubCurveIndex)
 	{
-		const ETransformCurveChannel Channel = (ETransformCurveChannel)SubCurveIndex;
+		const ETransformCurveChannel Channel = static_cast<ETransformCurveChannel>(SubCurveIndex);
 		for (int32 ChannelIndex = 0; ChannelIndex < 3; ++ChannelIndex)
 		{
-			const EVectorCurveChannel Axis = (EVectorCurveChannel)ChannelIndex;
+			const EVectorCurveChannel Axis = static_cast<EVectorCurveChannel>(ChannelIndex);
 			FAnimationCurveIdentifier TargetCurveIdentifier = CurveId;
 			UAnimationCurveIdentifierExtensions::GetTransformChildCurveIdentifier(TargetCurveIdentifier, Channel, Axis);
 			Controller->SetCurveKeys(TargetCurveIdentifier, SubCurveKeys[(SubCurveIndex * 3) + ChannelIndex], false);
@@ -314,7 +312,7 @@ FString FAddTransformCurveAction::ToStringInternal() const
 	return FText::Format(LOCTEXT("AddTransformCurveAction_Description", "Adding transform curve '{0}'."), FText::FromName(CurveId.InternalName.DisplayName)).ToString();
 }
 
-TUniquePtr<FChange> FAddRichCurveKeyAction::ExecuteInternal(UAnimDataModel* Model, UAnimDataController* Controller)
+TUniquePtr<FChange> FAddRichCurveKeyAction::ExecuteInternal(IAnimationDataModel* Model, IAnimationDataController* Controller)
 {
 	TUniquePtr<FChange> InverseAction = MakeUnique<FRemoveRichCurveKeyAction>(CurveId, Key.Time);
 	Controller->SetCurveKey(CurveId, Key, false);
@@ -329,7 +327,7 @@ FString FAddRichCurveKeyAction::ToStringInternal() const
 	return FText::Format(LOCTEXT("AddNamedRichCurveKeyAction_Description", "Adding key to {0} curve '{1}'."), FText::FromString(CurveId.CurveType == ERawCurveTrackTypes::RCT_Float ? FloatLabel : TransformLabel), FText::FromName(CurveId.InternalName.DisplayName)).ToString();
 }
 
-TUniquePtr<FChange> FSetRichCurveKeyAction::ExecuteInternal(UAnimDataModel* Model, UAnimDataController* Controller)
+TUniquePtr<FChange> FSetRichCurveKeyAction::ExecuteInternal(IAnimationDataModel* Model, IAnimationDataController* Controller)
 {
 	const FRichCurve& RichCurve = Model->GetRichCurve(CurveId);
 
@@ -350,7 +348,7 @@ FString FSetRichCurveKeyAction::ToStringInternal() const
 	return FText::Format(LOCTEXT("SetNamedRichCurveKeyAction_Description", "Setting key for {0} curve '{1}'."), FText::FromString(CurveId.CurveType == ERawCurveTrackTypes::RCT_Float ? FloatLabel : TransformLabel), FText::FromName(CurveId.InternalName.DisplayName)).ToString();
 }
 
-TUniquePtr<FChange> FRemoveRichCurveKeyAction::ExecuteInternal(UAnimDataModel* Model, UAnimDataController* Controller)
+TUniquePtr<FChange> FRemoveRichCurveKeyAction::ExecuteInternal(IAnimationDataModel* Model, IAnimationDataController* Controller)
 {	
 	const FRichCurve& RichCurve = Model->GetRichCurve(CurveId);
 
@@ -371,7 +369,7 @@ FString FRemoveRichCurveKeyAction::ToStringInternal() const
 	return FText::Format(LOCTEXT("RemoveNamedRichCurveKeyAction_Description", "Removing key from {0} curve '{1}'."), FText::FromString(CurveId.CurveType == ERawCurveTrackTypes::RCT_Float ? FloatLabel : TransformLabel), FText::FromName(CurveId.InternalName.DisplayName)).ToString();
 }
 
-TUniquePtr<FChange> FSetRichCurveKeysAction::ExecuteInternal(UAnimDataModel* Model, UAnimDataController* Controller)
+TUniquePtr<FChange> FSetRichCurveKeysAction::ExecuteInternal(IAnimationDataModel* Model, IAnimationDataController* Controller)
 {
 	const FRichCurve& RichCurve = Model->GetRichCurve(CurveId);
 
@@ -390,7 +388,7 @@ FString FSetRichCurveKeysAction::ToStringInternal() const
 	return FText::Format(LOCTEXT("SetNamedRichCurveKeysAction_Description", "Replacing keys for {0} curve '{1}'."), FText::FromString(CurveId.CurveType == ERawCurveTrackTypes::RCT_Float ? FloatLabel : TransformLabel), FText::FromName(CurveId.InternalName.DisplayName)).ToString();
 }
 
-TUniquePtr<FChange> FSetCurveColorAction::ExecuteInternal(UAnimDataModel* Model, UAnimDataController* Controller)
+TUniquePtr<FChange> FSetCurveColorAction::ExecuteInternal(IAnimationDataModel* Model, IAnimationDataController* Controller)
 {
 	const FAnimCurveBase& AnimationCurve = Model->GetCurve(CurveId);
 	const FLinearColor CurrentColor = AnimationCurve.Color;
@@ -410,7 +408,7 @@ FAddAtributeAction::FAddAtributeAction(const FAnimatedBoneAttribute& InAttribute
 	Keys = InAttribute.Curve.GetConstRefOfKeys();
 }
 
-TUniquePtr<FChange> FAddAtributeAction::ExecuteInternal(UAnimDataModel* Model, UAnimDataController* Controller)
+TUniquePtr<FChange> FAddAtributeAction::ExecuteInternal(IAnimationDataModel* Model, IAnimationDataController* Controller)
 {	
 	TArray<const void*> VoidValues;
 	Algo::Transform(Keys, VoidValues, [](const FAttributeKey& Key)
@@ -435,7 +433,7 @@ FString FAddAtributeAction::ToStringInternal() const
 	return FText::Format(LOCTEXT("AddAttributeAction_Description", "Adding attribute '{0}'."), FText::FromName(AttributeId.GetName())).ToString();
 }
 
-TUniquePtr<FChange> FRemoveAtributeAction::ExecuteInternal(UAnimDataModel* Model, UAnimDataController* Controller)
+TUniquePtr<FChange> FRemoveAtributeAction::ExecuteInternal(IAnimationDataModel* Model, IAnimationDataController* Controller)
 {
 	const FAnimatedBoneAttribute& Attribute = Model->GetAttribute(AttributeId);
 	TUniquePtr<FAddAtributeAction> InverseAction = MakeUnique<FAddAtributeAction>(Attribute);
@@ -450,7 +448,7 @@ FString FRemoveAtributeAction::ToStringInternal() const
 	return FText::Format(LOCTEXT("RemoveAttributeAction_Description", "Removing attribute '{0}'."), FText::FromName(AttributeId.GetName())).ToString();
 }
 
-TUniquePtr<FChange> FAddAtributeKeyAction::ExecuteInternal(UAnimDataModel* Model, UAnimDataController* Controller)
+TUniquePtr<FChange> FAddAtributeKeyAction::ExecuteInternal(IAnimationDataModel* Model, IAnimationDataController* Controller)
 {
 	Controller->SetAttributeKey(AttributeId, Key.Time, Key.GetValuePtr<void>(), AttributeId.GetType(), false);
 
@@ -462,7 +460,7 @@ FString FAddAtributeKeyAction::ToStringInternal() const
 	return FText::Format(LOCTEXT("AddAttributeKeyAction_Description", "Adding key to attribute '{0}'."), FText::FromName(AttributeId.GetName())).ToString();
 }
 
-TUniquePtr<FChange> FSetAtributeKeyAction::ExecuteInternal(UAnimDataModel* Model, UAnimDataController* Controller)
+TUniquePtr<FChange> FSetAtributeKeyAction::ExecuteInternal(IAnimationDataModel* Model, IAnimationDataController* Controller)
 {
 	const FAnimatedBoneAttribute& Attribute = Model->GetAttribute(AttributeId);
 
@@ -480,7 +478,7 @@ FString FSetAtributeKeyAction::ToStringInternal() const
 	return FText::Format(LOCTEXT("SetAttributeKeyAction_Description", "Setting key on attribute '{0}'."), FText::FromName(AttributeId.GetName())).ToString();
 }
 
-TUniquePtr<FChange> FRemoveAtributeKeyAction::ExecuteInternal(UAnimDataModel* Model, UAnimDataController* Controller)
+TUniquePtr<FChange> FRemoveAtributeKeyAction::ExecuteInternal(IAnimationDataModel* Model, IAnimationDataController* Controller)
 {
 	const FAnimatedBoneAttribute& Attribute = Model->GetAttribute(AttributeId);
 
@@ -504,7 +502,7 @@ FSetAtributeKeysAction::FSetAtributeKeysAction(const FAnimatedBoneAttribute& InA
 	Keys = InAttribute.Curve.GetConstRefOfKeys();
 }
 
-TUniquePtr<FChange> FSetAtributeKeysAction::ExecuteInternal(UAnimDataModel* Model, UAnimDataController* Controller)
+TUniquePtr<FChange> FSetAtributeKeysAction::ExecuteInternal(IAnimationDataModel* Model, IAnimationDataController* Controller)
 {
 	TUniquePtr<FSetAtributeKeysAction> InverseAction = MakeUnique<FSetAtributeKeysAction>(Model->GetAttribute(AttributeId));
 
