**New**

diff --git a/Engine/Source/Developer/AnimationDataController/Public/AnimDataController.h b/Engine/Source/Developer/AnimationDataController/Public/AnimDataController.h
--- a/Engine/Source/Developer/AnimationDataController/Public/AnimDataController.h
+++ b/Engine/Source/Developer/AnimationDataController/Public/AnimDataController.h
@@ -15,22 +15,7 @@
 
 #include "AnimDataController.generated.h"
 
-#if WITH_EDITOR
-
-#endif // WITH_EDITOR
-
-struct FAnimationCurveIdentifier;
-struct FAnimationAttributeIdentifier;
-
-namespace UE {
-namespace Anim {
-	class FOpenBracketAction;
-	class FCloseBracketAction;
-	
-	static const int32 DefaultCurveFlags = EAnimAssetCurveFlags::AACF_Editable;
-}}
-
-     UCLASS(BlueprintType)
+UCLASS(BlueprintType)
 class ANIMATIONDATACONTROLLER_API UAnimDataController : public UObject, public IAnimationDataController
 {
 	GENERATED_BODY()
@@ -44,15 +29,18 @@ public:
 
 #if WITH_EDITOR
 	/** Begin IAnimationDataController overrides */
-	virtual void SetModel(UAnimDataModel* InModel)  override;
-    virtual UAnimDataModel* GetModel() override { return Model; }
-	virtual const UAnimDataModel* const GetModel() const override { return Model; }
+	virtual void SetModel(TScriptInterface<IAnimationDataModel> InModel) override;
+    virtual TScriptInterface<IAnimationDataModel> GetModelInterface() const override { return ModelInterface; }
+	virtual const IAnimationDataModel* const GetModel() const override { return Model; }
 	virtual void OpenBracket(const FText& InTitle, bool bShouldTransact = true) override;
 	virtual void CloseBracket(bool bShouldTransact = true) override;
+	virtual void SetNumberOfFrames(FFrameNumber Length, bool bShouldTransact = true) override;
+	virtual void ResizeNumberOfFrames(FFrameNumber NewLength, FFrameNumber T0, FFrameNumber T1, bool bShouldTransact = true) override;
+	virtual void ResizeInFrames(FFrameNumber Length, FFrameNumber T0, FFrameNumber T1, bool bShouldTransact = true) override;
+	virtual void SetFrameRate(FFrameRate FrameRate, bool bShouldTransact = true) override;
 	virtual void SetPlayLength(float Length, bool bShouldTransact = true) override;
 	virtual void ResizePlayLength(float NewLength, float T0, float T1, bool bShouldTransact = true) override;
 	virtual void Resize(float Length, float T0, float T1, bool bShouldTransact = true) override;
-	virtual void SetFrameRate(FFrameRate FrameRate, bool bShouldTransact = true) override;
 	virtual int32 AddBoneTrack(FName BoneName, bool bShouldTransact = true) override;
 	virtual int32 InsertBoneTrack(FName BoneName, int32 DesiredIndex, bool bShouldTransact = true) override;
 	virtual bool RemoveBoneTrack(FName BoneName, bool bShouldTransact = true) override;
@@ -96,6 +84,9 @@ public:
 	}
     virtual bool RemoveAttributeKey(const FAnimationAttributeIdentifier& AttributeIdentifier, float Time, bool bShouldTransact = true) override;	
 	virtual bool DuplicateAttribute(const FAnimationAttributeIdentifier& AttributeIdentifier, const FAnimationAttributeIdentifier& NewAttributeIdentifier, bool bShouldTransact = true) override;
+	virtual void UpdateWithSkeleton(USkeleton* TargetSkeleton, bool bShouldTransact = true) override;
+	virtual void PopulateWithExistingModel(TScriptInterface<IAnimationDataModel> InModel) override;
+	virtual void InitializeModel() override {}
 protected:
 	virtual void NotifyBracketOpen() override;
 	virtual void NotifyBracketClosed() override;
@@ -106,48 +97,22 @@ private:
 	bool SetAttributeKey_Internal(const FAnimationAttributeIdentifier& AttributeIdentifier, float Time, const void* KeyValue, const UScriptStruct* TypeStruct, bool bShouldTransact = true);
 	bool SetAttributeKeys_Internal(const FAnimationAttributeIdentifier& AttributeIdentifier, TArrayView<const float> Times, TArrayView<const void*> KeyValues, const UScriptStruct* TypeStruct, bool bShouldTransact = true);
 
-	/** Returns whether or not the supplied curve type is supported by the controller functionality */
-	const bool IsSupportedCurveType(ERawCurveTrackTypes CurveType) const;
-	/** Returns the string representation of the provided curve enum type value */
-	FString GetCurveTypeValueName(ERawCurveTrackTypes InType) const;
-	
 	/** Resizes the curve/attribute data stored on the model according to the provided new length and time at which to insert or remove time */
 	void ResizeCurves(float NewLength, bool bInserted, float T0, float T1, bool bShouldTransact = true);
 	void ResizeAttributes(float NewLength, bool bInserted, float T0, float T1, bool bShouldTransact = true);
-
-	/** Ensures that a valid model is currently targeted */
-	void ValidateModel() const;
-
-	/** Verifies whether or not the Model's outer object is (or is derived from) the specified UClass */
-	bool CheckOuterClass(UClass* InClass) const;
-
-	/** Helper functionality to output script-based warnings and errors */
-	void ReportWarning(const FText& InMessage) const;
-	void ReportError(const FText& InMessage) const;
-
-	template <typename FmtType, typename... Types>
-    void ReportWarningf(const FmtType& Fmt, Types... Args) const
-	{	
-		ReportWarning(FText::Format(Fmt, Args...));
-	}
-
-	template <typename FmtType, typename... Types>
-    void ReportErrorf(const FmtType& Fmt, Types... Args) const
-	{
-		ReportError(FText::Format(Fmt, Args...));
-	}
 #endif // WITH_EDITOR
 
 private: 
 #if WITH_EDITOR
 	int32 BracketDepth;
-
-	UE::FChangeTransactor ChangeTransactor;	
 #endif // WITH_EDITOR
 
 #if WITH_EDITORONLY_DATA
 	UPROPERTY(transient)
 	TObjectPtr<UAnimDataModel> Model;
+	
+	UPROPERTY(transient)
+	TScriptInterface<IAnimationDataModel> ModelInterface;
 #endif // WITH_EDITORONLY_DATA
 
 	friend class FAnimDataControllerTestBase;
diff --git a/Engine/Source/Developer/AnimationDataController/Public/AnimDataControllerActions.h b/Engine/Source/Developer/AnimationDataController/Public/AnimDataControllerActions.h
new file mode 100644
--- /dev/null
+++ b/Engine/Source/Developer/AnimationDataController/Public/AnimDataControllerActions.h
@@ -0,0 +1,425 @@
+// Copyright Epic Games, Inc. All Rights Reserved.
+
+#pragma once
+
+#include "Misc/Change.h"
+#include "Misc/CoreMiscDefines.h"
+
+#include "Animation/AnimData/IAnimationDataModel.h"
+#include "Animation/AnimTypes.h" 
+#include "Misc/FrameRate.h"
+
+class UObject;
+class IAnimationDataModel;
+class IAnimationDataController;
+
+#if WITH_EDITOR
+
+namespace UE {
+
+namespace Anim {
+
+/**
+* UAnimDataController instanced FChange-based objects used for storing mutations to an IAnimationDataModel within the Transaction Buffer.
+* Each Action class represents an (invertable) operation mutating an IAnimationDataModel object utilizing a UAnimDataController. Allowing 
+* for a more granular approach to undo/redo-ing changes while also allowing for script-based interoperability.
+*/
+class ANIMATIONDATACONTROLLER_API FAnimDataBaseAction : public FSwapChange
+{
+public:
+	virtual TUniquePtr<FChange> Execute(UObject* Object) final;
+	virtual ~FAnimDataBaseAction() {}
+	virtual FString ToString() const override;
+
+protected:
+	virtual TUniquePtr<FChange> ExecuteInternal(IAnimationDataModel* Model, IAnimationDataController* Controller) = 0;
+	virtual FString ToStringInternal() const = 0;
+};
+
+class ANIMATIONDATACONTROLLER_API FOpenBracketAction : public FAnimDataBaseAction
+{
+public:
+	explicit FOpenBracketAction(const FString& InDescription) : Description(InDescription) {}
+	virtual ~FOpenBracketAction() {}
+protected:
+	FOpenBracketAction() {}
+	virtual TUniquePtr<FChange> ExecuteInternal(IAnimationDataModel* Model, IAnimationDataController* Controller) override;
+	virtual FString ToStringInternal() const override;
+
+protected:
+	FString Description;
+};
+
+class ANIMATIONDATACONTROLLER_API FCloseBracketAction : public FAnimDataBaseAction
+{
+public:
+	explicit FCloseBracketAction(const FString& InDescription) : Description(InDescription) {}
+	virtual ~FCloseBracketAction() {}
+protected:
+	FCloseBracketAction() {}
+	virtual TUniquePtr<FChange> ExecuteInternal(IAnimationDataModel* Model, IAnimationDataController* Controller) override;
+	virtual FString ToStringInternal() const override;
+
+protected:
+	FString Description;
+};
+
+class ANIMATIONDATACONTROLLER_API FAddTrackAction : public FAnimDataBaseAction
+{
+public:
+	explicit FAddTrackAction(const FBoneAnimationTrack& Track);
+	virtual ~FAddTrackAction() {}
+protected:
+	FAddTrackAction() {}
+	virtual TUniquePtr<FChange> ExecuteInternal(IAnimationDataModel* Model, IAnimationDataController* Controller) override;
+	virtual FString ToStringInternal() const override;
+
+protected:
+	FName Name;
+	int32 BoneTreeIndex;
+	FRawAnimSequenceTrack Data;
+};
+
+class ANIMATIONDATACONTROLLER_API FRemoveTrackAction : public FAnimDataBaseAction
+{
+public:
+	explicit FRemoveTrackAction(const FBoneAnimationTrack& Track);
+	virtual ~FRemoveTrackAction() {}
+protected:
+	FRemoveTrackAction() {}
+	virtual TUniquePtr<FChange> ExecuteInternal(IAnimationDataModel* Model, IAnimationDataController* Controller) override;
+	virtual FString ToStringInternal() const override;
+
+protected:
+	FName Name;
+};
+
+class ANIMATIONDATACONTROLLER_API FSetTrackKeysAction : public FAnimDataBaseAction
+{
+public:
+	explicit FSetTrackKeysAction(const FBoneAnimationTrack& Track);
+	virtual ~FSetTrackKeysAction() {}
+protected:
+	FSetTrackKeysAction() {}
+	virtual TUniquePtr<FChange> ExecuteInternal(IAnimationDataModel* Model, IAnimationDataController* Controller) override;
+	virtual FString ToStringInternal() const override;
+
+protected:
+	FName Name;
+	int32 TrackIndex;
+
+	FRawAnimSequenceTrack TrackData;
+};
+
+class ANIMATIONDATACONTROLLER_API FResizePlayLengthInFramesAction : public FAnimDataBaseAction
+{
+public:
+	explicit FResizePlayLengthInFramesAction(const IAnimationDataModel* InModel, FFrameNumber F0, FFrameNumber F1);
+	virtual ~FResizePlayLengthInFramesAction() override {}
+protected:
+	FResizePlayLengthInFramesAction() {}
+	virtual TUniquePtr<FChange> ExecuteInternal(IAnimationDataModel* Model, IAnimationDataController* Controller) override;
+	virtual FString ToStringInternal() const override;
+
+protected:
+	FFrameNumber Length;
+	FFrameNumber Frame0;
+	FFrameNumber Frame1;
+};
+
+class ANIMATIONDATACONTROLLER_API FSetFrameRateAction : public FAnimDataBaseAction
+{
+public:
+	explicit FSetFrameRateAction(const IAnimationDataModel* InModel);
+	virtual ~FSetFrameRateAction() override {}
+protected:
+	FSetFrameRateAction() {}
+	virtual TUniquePtr<FChange> ExecuteInternal(IAnimationDataModel* Model, IAnimationDataController* Controller) override;
+	virtual FString ToStringInternal() const override;
+
+protected:
+	FFrameRate FrameRate;
+};
+
+class ANIMATIONDATACONTROLLER_API FAddCurveAction : public FAnimDataBaseAction
+{
+public:
+	explicit FAddCurveAction(const FAnimationCurveIdentifier& InCurveId, int32 InFlags) : CurveId(InCurveId), Flags(InFlags) {}
+
+	virtual ~FAddCurveAction() {}
+protected:
+	FAddCurveAction() {}
+	virtual TUniquePtr<FChange> ExecuteInternal(IAnimationDataModel* Model, IAnimationDataController* Controller) override;
+	virtual FString ToStringInternal() const override;
+
+protected:
+	FAnimationCurveIdentifier CurveId;
+	int32 Flags;
+};
+
+class ANIMATIONDATACONTROLLER_API FAddFloatCurveAction : public FAnimDataBaseAction
+{
+public:
+	explicit FAddFloatCurveAction(const FAnimationCurveIdentifier& InCurveId, int32 InFlags, const TArray<FRichCurveKey>& InKeys, const FLinearColor& InColor) : CurveId(InCurveId), Flags(InFlags), Keys(InKeys), Color(InColor) {}
+	virtual ~FAddFloatCurveAction() {}
+protected:
+	FAddFloatCurveAction() {}
+	virtual TUniquePtr<FChange> ExecuteInternal(IAnimationDataModel* Model, IAnimationDataController* Controller) override;
+	virtual FString ToStringInternal() const override;
+
+protected:
+	FAnimationCurveIdentifier CurveId;
+	int32 Flags;
+	TArray<FRichCurveKey> Keys;
+	FLinearColor Color;
+};
+
+class ANIMATIONDATACONTROLLER_API FAddTransformCurveAction : public FAnimDataBaseAction
+{
+public:
+	explicit FAddTransformCurveAction(const FAnimationCurveIdentifier& InCurveId, int32 InFlags, const FTransformCurve& InTransformCurve);
+	virtual ~FAddTransformCurveAction() {}
+protected:
+	FAddTransformCurveAction() {}
+	virtual TUniquePtr<FChange> ExecuteInternal(IAnimationDataModel* Model, IAnimationDataController* Controller) override;
+	virtual FString ToStringInternal() const override;
+
+protected:
+	FAnimationCurveIdentifier CurveId;
+	int32 Flags;
+	
+	TArray<FRichCurveKey> SubCurveKeys[9];
+};
+
+class ANIMATIONDATACONTROLLER_API FRemoveCurveAction : public FAnimDataBaseAction
+{
+public:
+	explicit FRemoveCurveAction(const FAnimationCurveIdentifier& InCurveId) : CurveId(InCurveId) {}
+	virtual ~FRemoveCurveAction() {}
+protected:
+	FRemoveCurveAction() {}
+	virtual TUniquePtr<FChange> ExecuteInternal(IAnimationDataModel* Model, IAnimationDataController* Controller) override;
+	virtual FString ToStringInternal() const override;
+
+protected:
+	FAnimationCurveIdentifier CurveId;
+};
+
+class ANIMATIONDATACONTROLLER_API FSetCurveFlagsAction : public FAnimDataBaseAction
+{
+public:
+	explicit FSetCurveFlagsAction(const FAnimationCurveIdentifier& InCurveId, int32 InFlags, ERawCurveTrackTypes InCurveType) : CurveId(InCurveId), Flags(InFlags), CurveType(InCurveType) {}
+	virtual ~FSetCurveFlagsAction() {}
+protected:
+	FSetCurveFlagsAction() {}
+	virtual TUniquePtr<FChange> ExecuteInternal(IAnimationDataModel* Model, IAnimationDataController* Controller) override;
+	virtual FString ToStringInternal() const override;
+
+protected:
+	FAnimationCurveIdentifier CurveId;
+	int32 Flags;
+	ERawCurveTrackTypes CurveType;
+};
+
+class ANIMATIONDATACONTROLLER_API FRenameCurveAction : public FAnimDataBaseAction
+{
+public:
+	explicit FRenameCurveAction(const FAnimationCurveIdentifier& InCurveId, const FAnimationCurveIdentifier& InNewCurveId) : CurveId(InCurveId), NewCurveId(InNewCurveId) {}
+	virtual ~FRenameCurveAction() {}
+protected:
+	FRenameCurveAction() {}
+	virtual TUniquePtr<FChange> ExecuteInternal(IAnimationDataModel* Model, IAnimationDataController* Controller) override;
+	virtual FString ToStringInternal() const override;
+
+protected:
+	FAnimationCurveIdentifier CurveId;
+	FAnimationCurveIdentifier NewCurveId;
+};
+
+class ANIMATIONDATACONTROLLER_API FScaleCurveAction : public FAnimDataBaseAction
+{
+public:
+	explicit FScaleCurveAction(const FAnimationCurveIdentifier& InCurveId, float InOrigin, float InFactor, ERawCurveTrackTypes InCurveType);
+	virtual ~FScaleCurveAction() {}
+protected:
+	FScaleCurveAction() {}
+	virtual TUniquePtr<FChange> ExecuteInternal(IAnimationDataModel* Model, IAnimationDataController* Controller) override;
+	virtual FString ToStringInternal() const override;
+
+protected:
+	FAnimationCurveIdentifier CurveId;
+	ERawCurveTrackTypes CurveType;
+	float Origin;
+	float Factor;
+};
+
+class ANIMATIONDATACONTROLLER_API FAddRichCurveKeyAction : public FAnimDataBaseAction
+{
+public:
+	explicit FAddRichCurveKeyAction(const FAnimationCurveIdentifier& InCurveId, const FRichCurveKey& InKey) : CurveId(InCurveId), Key(InKey) {}
+	virtual ~FAddRichCurveKeyAction() {}
+protected:
+	FAddRichCurveKeyAction() {}
+	virtual TUniquePtr<FChange> ExecuteInternal(IAnimationDataModel* Model, IAnimationDataController* Controller) override;
+	virtual FString ToStringInternal() const override;
+
+protected:
+	FAnimationCurveIdentifier CurveId;
+	ERawCurveTrackTypes CurveType;
+	FRichCurveKey Key;
+};
+
+class ANIMATIONDATACONTROLLER_API FSetRichCurveKeyAction : public FAnimDataBaseAction
+{
+public:
+	explicit FSetRichCurveKeyAction(const FAnimationCurveIdentifier& InCurveId, const FRichCurveKey& InKey) : CurveId(InCurveId), Key(InKey) {}
+	virtual ~FSetRichCurveKeyAction() {}
+protected:
+	FSetRichCurveKeyAction() {}
+	virtual TUniquePtr<FChange> ExecuteInternal(IAnimationDataModel* Model, IAnimationDataController* Controller) override;
+	virtual FString ToStringInternal() const override;
+
+protected:
+	FAnimationCurveIdentifier CurveId;
+	FRichCurveKey Key;
+};
+
+class ANIMATIONDATACONTROLLER_API FRemoveRichCurveKeyAction : public FAnimDataBaseAction
+{
+public:
+	explicit FRemoveRichCurveKeyAction(const FAnimationCurveIdentifier& InCurveId, const float InTime) : CurveId(InCurveId), Time(InTime) {}
+	virtual ~FRemoveRichCurveKeyAction() {}
+protected:
+	FRemoveRichCurveKeyAction() {}
+	virtual TUniquePtr<FChange> ExecuteInternal(IAnimationDataModel* Model, IAnimationDataController* Controller) override;
+	virtual FString ToStringInternal() const override;
+
+protected:
+	FAnimationCurveIdentifier CurveId;
+	float Time;
+};
+
+class ANIMATIONDATACONTROLLER_API FSetRichCurveKeysAction : public FAnimDataBaseAction
+{
+public:
+	explicit FSetRichCurveKeysAction(const FAnimationCurveIdentifier& InCurveId, const TArray<FRichCurveKey>& InKeys) : CurveId(InCurveId), Keys(InKeys) {}
+	virtual ~FSetRichCurveKeysAction() {}
+protected:
+	FSetRichCurveKeysAction() {}
+	virtual TUniquePtr<FChange> ExecuteInternal(IAnimationDataModel* Model, IAnimationDataController* Controller) override;
+	virtual FString ToStringInternal() const override;
+
+protected:
+	FAnimationCurveIdentifier CurveId;
+	TArray<FRichCurveKey> Keys;
+};
+
+class ANIMATIONDATACONTROLLER_API FSetCurveColorAction : public FAnimDataBaseAction
+{
+public:
+	explicit FSetCurveColorAction(const FAnimationCurveIdentifier& InCurveId, const FLinearColor& InColor) : CurveId(InCurveId), Color(InColor) {}
+	virtual ~FSetCurveColorAction() {}
+protected:
+	FSetCurveColorAction() {}
+	virtual TUniquePtr<FChange> ExecuteInternal(IAnimationDataModel* Model, IAnimationDataController* Controller) override;
+	virtual FString ToStringInternal() const override;
+
+protected:
+	FAnimationCurveIdentifier CurveId;
+	FLinearColor Color;
+};
+
+class ANIMATIONDATACONTROLLER_API FAddAtributeAction : public FAnimDataBaseAction
+{
+public:
+	explicit FAddAtributeAction(const FAnimatedBoneAttribute& InAttribute);
+	virtual ~FAddAtributeAction() {}
+protected:
+	FAddAtributeAction() {}
+	virtual TUniquePtr<FChange> ExecuteInternal(IAnimationDataModel* Model, IAnimationDataController* Controller) override;
+	virtual FString ToStringInternal() const override;
+
+protected:
+	FAnimationAttributeIdentifier AttributeId;
+	TArray<FAttributeKey> Keys;
+};
+
+class ANIMATIONDATACONTROLLER_API FRemoveAtributeAction : public FAnimDataBaseAction
+{
+public:
+	explicit FRemoveAtributeAction(const FAnimationAttributeIdentifier& InAttributeId) : AttributeId(InAttributeId) {}
+	virtual ~FRemoveAtributeAction() {}
+protected:
+	FRemoveAtributeAction() {}
+	virtual TUniquePtr<FChange> ExecuteInternal(IAnimationDataModel* Model, IAnimationDataController* Controller) override;
+	virtual FString ToStringInternal() const override;
+
+protected:
+	FAnimationAttributeIdentifier AttributeId;
+};
+
+class ANIMATIONDATACONTROLLER_API FAddAtributeKeyAction : public FAnimDataBaseAction
+{
+public:
+	explicit FAddAtributeKeyAction(const FAnimationAttributeIdentifier& InAttributeId, const FAttributeKey& InKey) : AttributeId(InAttributeId), Key(InKey) {}
+	virtual ~FAddAtributeKeyAction() {}
+protected:
+	FAddAtributeKeyAction() {}
+	virtual TUniquePtr<FChange> ExecuteInternal(IAnimationDataModel* Model, IAnimationDataController* Controller) override;
+	virtual FString ToStringInternal() const override;
+
+protected:
+	FAnimationAttributeIdentifier AttributeId;
+	FAttributeKey Key;
+};
+
+class ANIMATIONDATACONTROLLER_API FSetAtributeKeyAction : public FAnimDataBaseAction
+{
+public:
+	explicit FSetAtributeKeyAction(const FAnimationAttributeIdentifier& InAttributeId, const FAttributeKey& InKey) : AttributeId(InAttributeId), Key(InKey) {}
+	virtual ~FSetAtributeKeyAction() {}
+protected:
+	FSetAtributeKeyAction() {}
+	virtual TUniquePtr<FChange> ExecuteInternal(IAnimationDataModel* Model, IAnimationDataController* Controller) override;
+	virtual FString ToStringInternal() const override;
+
+protected:
+	FAnimationAttributeIdentifier AttributeId;
+	FAttributeKey Key;
+};
+
+class ANIMATIONDATACONTROLLER_API FRemoveAtributeKeyAction : public FAnimDataBaseAction
+{
+public:
+	explicit FRemoveAtributeKeyAction(const FAnimationAttributeIdentifier& InAttributeId, float InTime) : AttributeId(InAttributeId), Time(InTime) {}
+	virtual ~FRemoveAtributeKeyAction() {}
+protected:
+	FRemoveAtributeKeyAction() {}
+	virtual TUniquePtr<FChange> ExecuteInternal(IAnimationDataModel* Model, IAnimationDataController* Controller) override;
+	virtual FString ToStringInternal() const override;
+
+protected:
+	FAnimationAttributeIdentifier AttributeId;
+	float Time;
+};
+
+class ANIMATIONDATACONTROLLER_API FSetAtributeKeysAction : public FAnimDataBaseAction
+{
+public:
+	explicit FSetAtributeKeysAction(const FAnimatedBoneAttribute& InAttribute);
+	virtual ~FSetAtributeKeysAction() {}
+protected:
+	FSetAtributeKeysAction() {}
+	virtual TUniquePtr<FChange> ExecuteInternal(IAnimationDataModel* Model, IAnimationDataController* Controller) override;
+	virtual FString ToStringInternal() const override;
+
+protected:
+	FAnimationAttributeIdentifier AttributeId;
+	TArray<FAttributeKey> Keys;
+};
+
+} // namespace Anim
+
+} // namespace UE
+
+#endif // WITH_EDITOR
\ No newline at end of file
diff --git a/Engine/Source/Developer/MeshBoneReduction/Private/MeshBoneReduction.cpp b/Engine/Source/Developer/MeshBoneReduction/Private/MeshBoneReduction.cpp
--- a/Engine/Source/Developer/MeshBoneReduction/Private/MeshBoneReduction.cpp
+++ b/Engine/Source/Developer/MeshBoneReduction/Private/MeshBoneReduction.cpp
@@ -51,7 +51,7 @@ public:
 		}
 	}
 
-	bool GetBoneReductionData(const USkeletalMesh* SkeletalMesh, int32 DesiredLOD, TMap<FBoneIndexType, FBoneIndexType>& OutBonesToReplace, const TArray<FName>* BoneNamesToRemove = NULL) override
+	bool GetBoneReductionData(const USkeletalMesh* SkeletalMesh, int32 DesiredLOD, TMap<FBoneIndexType, FBoneIndexType>& OutBonesToReplace, const TArray<FName>* BoneNamesToRemove = nullptr) override
 	{
 		if (!SkeletalMesh)
 		{
@@ -320,14 +320,20 @@ public:
 			Pose.SetBoneContainer(&RequiredBones);
 			Pose.ResetToRefPose();
 
+			FBlendedCurve TempCurve;
+			TempCurve.InitFrom(RequiredBones);
+			UE::Anim::FStackAttributeContainer TempAttributes;
+
+			FAnimationPoseData AnimPoseData(Pose, TempCurve, TempAttributes);
+
 			// Get component space retarget base pose, will be equivalent of ref-pose if not edited
 			TArray<FTransform> ComponentSpaceRefPose;
-			FAnimationRuntime::FillUpComponentSpaceTransformsRetargetBasePose(SkeletalMesh, ComponentSpaceRefPose);
+			FAnimationRuntime::FillUpComponentSpaceTransforms(SkeletalMesh->GetRefSkeleton(), SkeletalMesh->GetRefSkeleton().GetRefBonePose(), ComponentSpaceRefPose);
 			
 			// Retrieve animated pose from anim sequence (including retargeting)
 			const USkeleton* Skeleton = SkeletalMesh->GetSkeleton();
 			const FName RetargetSource = Skeleton->GetRetargetSourceForMesh(SkeletalMesh);
-			UE::Anim::BuildPoseFromModel(BakePoseAnim->GetDataModel(), Pose, 0.f, EAnimInterpolationType::Step, RetargetSource, Skeleton->GetRefLocalPoses(RetargetSource));
+			UE::Anim::BuildPoseFromModel(BakePoseAnim->GetDataModel(), AnimPoseData, 0.0, EAnimInterpolationType::Step, RetargetSource, Skeleton->GetRefLocalPoses(RetargetSource));
 			
 			// Calculate component space animated pose matrices
 			TArray<FMatrix> ComponentSpaceAnimatedPose;
diff --git a/Engine/Source/Editor/AnimGraph/Private/AnimPreviewInstance.cpp b/Engine/Source/Editor/AnimGraph/Private/AnimPreviewInstance.cpp
--- a/Engine/Source/Editor/AnimGraph/Private/AnimPreviewInstance.cpp
+++ b/Engine/Source/Editor/AnimGraph/Private/AnimPreviewInstance.cpp
@@ -4,6 +4,7 @@
 #include "AnimPreviewInstance.h"
 #include "Animation/DebugSkelMeshComponent.h"
 #include "AnimationRuntime.h"
+#include "Animation/AnimationSettings.h"
 #include "Animation/AnimSequence.h"
 #include "Animation/AnimSequenceHelpers.h"
 
@@ -793,7 +794,12 @@ void UAnimPreviewInstance::MontagePreview_StepBackward()
 
 float UAnimPreviewInstance::MontagePreview_CalculateStepLength()
 {
-	return 1.0f / 30.0f;
+	if (UAnimMontage* Montage = Cast<UAnimMontage>(CurrentAsset))
+	{
+		return Montage->GetSamplingFrameRate().AsInterval();
+	}
+
+	return UAnimationSettings::Get()->GetDefaultFrameRate().AsInterval();
 }
 
 void UAnimPreviewInstance::MontagePreview_JumpToStart()
diff --git a/Engine/Source/Editor/AnimationBlueprintLibrary/Private/AnimPose.cpp b/Engine/Source/Editor/AnimationBlueprintLibrary/Private/AnimPose.cpp
--- a/Engine/Source/Editor/AnimationBlueprintLibrary/Private/AnimPose.cpp
+++ b/Engine/Source/Editor/AnimationBlueprintLibrary/Private/AnimPose.cpp
@@ -499,12 +499,11 @@ float UAnimPoseExtensions::GetCurveWeight(const FAnimPose& Pose, const FName& Cu
 
 void UAnimPoseExtensions::GetAnimPoseAtFrame(const UAnimSequenceBase* AnimationSequenceBase, int32 FrameIndex, FAnimPoseEvaluationOptions EvaluationOptions, FAnimPose& Pose)
 {
-	float Time = 0.f;
-	UAnimationBlueprintLibrary::GetTimeAtFrame(AnimationSequenceBase, FrameIndex, Time);
+	const double Time = AnimationSequenceBase->GetDataModel()->GetFrameRate().AsSeconds(FrameIndex);
 	GetAnimPoseAtTime(AnimationSequenceBase, Time, EvaluationOptions, Pose);
 }
 
-void UAnimPoseExtensions::GetAnimPoseAtTime(const UAnimSequenceBase* AnimationSequenceBase, float Time, FAnimPoseEvaluationOptions EvaluationOptions, FAnimPose& Pose)
+void UAnimPoseExtensions::GetAnimPoseAtTime(const UAnimSequenceBase* AnimationSequenceBase, double Time, FAnimPoseEvaluationOptions EvaluationOptions, FAnimPose& Pose)
 {
 	TArray<FAnimPose> InOutPoses;
 	GetAnimPoseAtTimeIntervals(AnimationSequenceBase, { Time }, EvaluationOptions, InOutPoses);
@@ -516,7 +515,7 @@ void UAnimPoseExtensions::GetAnimPoseAtTime(const UAnimSequenceBase* AnimationSe
 	}
 } 
 
-void UAnimPoseExtensions::GetAnimPoseAtTimeIntervals(const UAnimSequenceBase* AnimationSequenceBase, TArray<float> TimeIntervals, FAnimPoseEvaluationOptions EvaluationOptions, TArray<FAnimPose>& InOutPoses)
+void UAnimPoseExtensions::GetAnimPoseAtTimeIntervals(const UAnimSequenceBase* AnimationSequenceBase, TArray<double> TimeIntervals, FAnimPoseEvaluationOptions EvaluationOptions, TArray<FAnimPose>& InOutPoses)
 {
 	if (AnimationSequenceBase && AnimationSequenceBase->GetSkeleton())
 	{
@@ -556,7 +555,7 @@ void UAnimPoseExtensions::GetAnimPoseAtTimeIntervals(const UAnimSequenceBase* An
         UE::Anim::FStackAttributeContainer Attributes;
 
         FAnimationPoseData PoseData(CompactPose, Curve, Attributes);
-        FAnimExtractContext Context(0.f, EvaluationOptions.bExtractRootMotion);
+        FAnimExtractContext Context(0.0, EvaluationOptions.bExtractRootMotion);
     
         FCompactPose BasePose;
         BasePose.SetBoneContainer(&RequiredBones);
@@ -569,7 +568,7 @@ void UAnimPoseExtensions::GetAnimPoseAtTimeIntervals(const UAnimSequenceBase* An
 		
 		for (int32 Index = 0; Index < TimeIntervals.Num(); ++Index)
 		{
-			const float EvalInterval = TimeIntervals[Index];
+			const double EvalInterval = TimeIntervals[Index];
 			
 			bool bValidTime = false;
 			UAnimationBlueprintLibrary::IsValidTime(AnimationSequenceBase, EvalInterval, bValidTime);
diff --git a/Engine/Source/Editor/AnimationBlueprintLibrary/Private/AnimationAttributeBlueprintLibrary.cpp b/Engine/Source/Editor/AnimationBlueprintLibrary/Private/AnimationAttributeBlueprintLibrary.cpp
--- a/Engine/Source/Editor/AnimationBlueprintLibrary/Private/AnimationAttributeBlueprintLibrary.cpp
+++ b/Engine/Source/Editor/AnimationBlueprintLibrary/Private/AnimationAttributeBlueprintLibrary.cpp
@@ -131,14 +131,14 @@ DEFINE_FUNCTION(UAnimationAttributeBlueprintLibrary::execSetAttributeKeys)
 }
 
 
-bool UAnimationAttributeBlueprintLibrary::GetAttributeKey(UAnimDataModel* AnimationDataModel, const FAnimationAttributeIdentifier& AttributeIdentifier, float Time, int32& Value)
+bool UAnimationAttributeBlueprintLibrary::GetAttributeKey(TScriptInterface<IAnimationDataModel> AnimationDataModel, const FAnimationAttributeIdentifier& AttributeIdentifier, float Time, int32& Value)
 {
 	// We should never hit this! Stubbed to avoid NoExport on the class.
 	check(0);
 	return false;
 }
 
-bool UAnimationAttributeBlueprintLibrary::Generic_GetAttributeKey(UAnimDataModel* AnimationDataModel, const FAnimationAttributeIdentifier& AttributeIdentifier, float Time, UScriptStruct* ScriptStruct, void* ValuePtr)
+bool UAnimationAttributeBlueprintLibrary::Generic_GetAttributeKey(TScriptInterface<IAnimationDataModel> AnimationDataModel, const FAnimationAttributeIdentifier& AttributeIdentifier, float Time, UScriptStruct* ScriptStruct, void* ValuePtr)
 {
 	const FAnimatedBoneAttribute* AttributePtr = AnimationDataModel->FindAttribute(AttributeIdentifier);
 
@@ -164,7 +164,7 @@ bool UAnimationAttributeBlueprintLibrary::Generic_GetAttributeKey(UAnimDataModel
 
 DEFINE_FUNCTION(UAnimationAttributeBlueprintLibrary::execGetAttributeKey)
 {
-	P_GET_OBJECT(UAnimDataModel, AnimationDataModel);
+	P_GET_TINTERFACE(IAnimationDataModel, AnimationDataModel);
 	P_GET_STRUCT_REF(FAnimationAttributeIdentifier, AttributeIdentifier);
 	P_GET_PROPERTY(FFloatProperty, TimeInterval);
 
@@ -186,7 +186,7 @@ DEFINE_FUNCTION(UAnimationAttributeBlueprintLibrary::execGetAttributeKey)
 		FBlueprintCoreDelegates::ThrowScriptException(P_THIS, Stack, ExceptionInfo);
 	}
 
-	if (!AnimationDataModel)
+	if (!AnimationDataModel.GetObject())
 	{
 		FBlueprintExceptionInfo ExceptionInfo(
 			EBlueprintExceptionType::AccessViolation, 
@@ -208,14 +208,14 @@ DEFINE_FUNCTION(UAnimationAttributeBlueprintLibrary::execGetAttributeKey)
 }
 
 
-bool UAnimationAttributeBlueprintLibrary::GetAttributeKeys(UAnimDataModel* AnimationDataModel, const FAnimationAttributeIdentifier& AttributeIdentifier, TArray<float>& OutTimes, TArray<int32>& OutValues)
+bool UAnimationAttributeBlueprintLibrary::GetAttributeKeys(TScriptInterface<IAnimationDataModel> AnimationDataModel, const FAnimationAttributeIdentifier& AttributeIdentifier, TArray<float>& OutTimes, TArray<int32>& OutValues)
 {
 	// We should never hit this! Stubbed to avoid NoExport on the class.
 	check(0);
 	return false;
 }
 
-bool UAnimationAttributeBlueprintLibrary::Generic_GetAttributeKeys(UAnimDataModel* AnimationDataModel, const FAnimationAttributeIdentifier& AttributeIdentifier, TArray<float>& Times, void* ValuesArray, const FArrayProperty* ValuesArrayProperty)
+bool UAnimationAttributeBlueprintLibrary::Generic_GetAttributeKeys(TScriptInterface<IAnimationDataModel> AnimationDataModel, const FAnimationAttributeIdentifier& AttributeIdentifier, TArray<float>& Times, void* ValuesArray, const FArrayProperty* ValuesArrayProperty)
 {
 	const FAnimatedBoneAttribute* AttributePtr = AnimationDataModel->FindAttribute(AttributeIdentifier);
 
@@ -249,7 +249,7 @@ bool UAnimationAttributeBlueprintLibrary::Generic_GetAttributeKeys(UAnimDataMode
 
 DEFINE_FUNCTION(UAnimationAttributeBlueprintLibrary::execGetAttributeKeys)
 {
-	P_GET_OBJECT(UAnimDataModel, AnimationDataModel);
+	P_GET_TINTERFACE(IAnimationDataModel, AnimationDataModel);
 	P_GET_STRUCT_REF(FAnimationAttributeIdentifier, AttributeIdentifier);
 	P_GET_TARRAY_REF(float, TimeIntervals);
 
@@ -271,7 +271,7 @@ DEFINE_FUNCTION(UAnimationAttributeBlueprintLibrary::execGetAttributeKeys)
 		FBlueprintCoreDelegates::ThrowScriptException(P_THIS, Stack, ExceptionInfo);
 	}
 	
-	if (!AnimationDataModel)
+	if (!AnimationDataModel.GetObject())
 	{
 		FBlueprintExceptionInfo ExceptionInfo(
 			EBlueprintExceptionType::AccessViolation, 
diff --git a/Engine/Source/Editor/AnimationBlueprintLibrary/Private/AnimationBlueprintLibrary.cpp b/Engine/Source/Editor/AnimationBlueprintLibrary/Private/AnimationBlueprintLibrary.cpp
--- a/Engine/Source/Editor/AnimationBlueprintLibrary/Private/AnimationBlueprintLibrary.cpp
+++ b/Engine/Source/Editor/AnimationBlueprintLibrary/Private/AnimationBlueprintLibrary.cpp
@@ -24,7 +24,7 @@
 #include "Modules/ModuleInterface.h"
 
 #include "AnimationRuntime.h"
-#include "Animation/AnimData/AnimDataModel.h"
+#include "Animation/AnimData/IAnimationDataModel.h"
 #include "Animation/AnimSequenceHelpers.h"
 #include "AnimGraphNode_AssetPlayerBase.h"
 #include "Kismet2/BlueprintEditorUtils.h"
@@ -2197,7 +2197,7 @@ bool UAnimationBlueprintLibrary::EvaluateRootBoneTimecodeAttributesAtTime(const
 		return false;
 	}
 
-	const UAnimDataModel* AnimDataModel = AnimationSequenceBase->GetDataModel();
+	const IAnimationDataModel* AnimDataModel = AnimationSequenceBase->GetDataModel();
 	if (!AnimDataModel)
 	{
 		return false;
@@ -2386,7 +2386,7 @@ bool UAnimationBlueprintLibrary::EvaluateRootBoneTimecodeSubframeAttributeAtTime
 		return false;
 	}
 
-	const UAnimDataModel* AnimDataModel = AnimationSequenceBase->GetDataModel();
+	const IAnimationDataModel* AnimDataModel = AnimationSequenceBase->GetDataModel();
 	if (!AnimDataModel)
 	{
 		return false;
diff --git a/Engine/Source/Editor/AnimationBlueprintLibrary/Public/AnimPose.h b/Engine/Source/Editor/AnimationBlueprintLibrary/Public/AnimPose.h
--- a/Engine/Source/Editor/AnimationBlueprintLibrary/Public/AnimPose.h
+++ b/Engine/Source/Editor/AnimationBlueprintLibrary/Public/AnimPose.h
@@ -229,7 +229,7 @@ public:
 	* @param	Pose					Anim Pose to hold the evaluated data
 	*/
 	UFUNCTION(BlueprintCallable, meta = (ScriptMethod), Category = "Animation|Pose")
-	static void GetAnimPoseAtTime(const UAnimSequenceBase* AnimationSequenceBase, float Time, FAnimPoseEvaluationOptions EvaluationOptions, FAnimPose& Pose);
+	static void GetAnimPoseAtTime(const UAnimSequenceBase* AnimationSequenceBase, double Time, FAnimPoseEvaluationOptions EvaluationOptions, FAnimPose& Pose);
 
 	/**
 	* Evaluates an Animation Sequence Base at different time intervals to generate a valid Anim Pose instances
@@ -239,7 +239,7 @@ public:
 	* @param	EvaluationOptions		Options determining the way the pose should be evaluated
 	* @param	InOutPoses				Anim Poses holding the evaluated data (number matches TimeIntervals)
 	*/
-	static void GetAnimPoseAtTimeIntervals(const UAnimSequenceBase* AnimationSequenceBase, TArray<float> TimeIntervals, FAnimPoseEvaluationOptions EvaluationOptions, TArray<FAnimPose>& InOutPoses);
+	static void GetAnimPoseAtTimeIntervals(const UAnimSequenceBase* AnimationSequenceBase, TArray<double> TimeIntervals, FAnimPoseEvaluationOptions EvaluationOptions, TArray<FAnimPose>& InOutPoses);
 
 	/**
 	* Evaluates an Animation Sequence Base to generate a valid Anim Pose instance
diff --git a/Engine/Source/Editor/AnimationBlueprintLibrary/Public/AnimationAttributeBlueprintLibrary.h b/Engine/Source/Editor/AnimationBlueprintLibrary/Public/AnimationAttributeBlueprintLibrary.h
--- a/Engine/Source/Editor/AnimationBlueprintLibrary/Public/AnimationAttributeBlueprintLibrary.h
+++ b/Engine/Source/Editor/AnimationBlueprintLibrary/Public/AnimationAttributeBlueprintLibrary.h
@@ -26,12 +26,12 @@ public:
 
 
 	UFUNCTION(BlueprintCallable, CustomThunk, Category = "Animation|Attributes", meta=(CustomStructureParam="Value", BlueprintInternalUseOnly="true"))
-	static bool GetAttributeKey(UAnimDataModel* AnimationDataModel, const FAnimationAttributeIdentifier& AttributeIdentifier, float Time, int32& Value);
-	static bool Generic_GetAttributeKey(UAnimDataModel* AnimationDataModel, const FAnimationAttributeIdentifier& AttributeIdentifier, float Time, UScriptStruct* ScriptStruct, void* Value);
+	static bool GetAttributeKey(TScriptInterface<IAnimationDataModel> AnimationDataModel, const FAnimationAttributeIdentifier& AttributeIdentifier, float Time, int32& Value);
+	static bool Generic_GetAttributeKey(TScriptInterface<IAnimationDataModel> AnimationDataModel, const FAnimationAttributeIdentifier& AttributeIdentifier, float Time, UScriptStruct* ScriptStruct, void* Value);
 	DECLARE_FUNCTION(execGetAttributeKey);
 
 	UFUNCTION(BlueprintCallable, CustomThunk, Category = "Animation|Attributes", meta=(CustomStructureParam="Values", BlueprintInternalUseOnly="true"))
-	static bool GetAttributeKeys(UAnimDataModel* AnimationDataModel, const FAnimationAttributeIdentifier& AttributeIdentifier, TArray<float>& OutTimes, TArray<int32>& Values);
-	static bool Generic_GetAttributeKeys(UAnimDataModel* AnimationDataModel, const FAnimationAttributeIdentifier& AttributeIdentifier, TArray<float>& Times, void* ValuesArray, const FArrayProperty* ValuesArrayProperty);
+	static bool GetAttributeKeys(TScriptInterface<IAnimationDataModel> AnimationDataModel, const FAnimationAttributeIdentifier& AttributeIdentifier, TArray<float>& OutTimes, TArray<int32>& Values);
+	static bool Generic_GetAttributeKeys(TScriptInterface<IAnimationDataModel> AnimationDataModel, const FAnimationAttributeIdentifier& AttributeIdentifier, TArray<float>& Times, void* ValuesArray, const FArrayProperty* ValuesArrayProperty);
 	DECLARE_FUNCTION(execGetAttributeKeys);
 };
diff --git a/Engine/Source/Editor/DetailCustomizations/Private/PoseAssetDetails.cpp b/Engine/Source/Editor/DetailCustomizations/Private/PoseAssetDetails.cpp
--- a/Engine/Source/Editor/DetailCustomizations/Private/PoseAssetDetails.cpp
+++ b/Engine/Source/Editor/DetailCustomizations/Private/PoseAssetDetails.cpp
@@ -668,7 +668,7 @@ bool FPoseAssetDetails::IsUpdateSourceEnabled() const
 	const UAnimSequence* AnimSequenceSelected = Cast<UAnimSequence>(ObjectSet);
 	if (AnimSequenceSelected && PoseAsset.IsValid())
 	{
-		return (PoseAsset->SourceAnimation && PoseAsset->SourceAnimation->GetRawDataGuid() != PoseAsset->SourceAnimationRawDataGUID);		
+		return (PoseAsset->SourceAnimation && PoseAsset->SourceAnimation->GetDataModel()->GenerateGuid() != PoseAsset->SourceAnimationRawDataGUID);		
 	}
 
 	return false;
diff --git a/Engine/Source/Editor/MovieSceneTools/Private/EditModes/SkeletalAnimationTrackEditMode.cpp b/Engine/Source/Editor/MovieSceneTools/Private/EditModes/SkeletalAnimationTrackEditMode.cpp
--- a/Engine/Source/Editor/MovieSceneTools/Private/EditModes/SkeletalAnimationTrackEditMode.cpp
+++ b/Engine/Source/Editor/MovieSceneTools/Private/EditModes/SkeletalAnimationTrackEditMode.cpp
@@ -364,7 +364,7 @@ void FSkeletalAnimationTrackEditMode::Render(const FSceneView* View, FViewport*
 											UE::Anim::FStackAttributeContainer TempAttributes;
 											FAnimationPoseData OutAnimationPoseData(OutPose, OutCurve, TempAttributes);
 
-											const float Seconds = static_cast<float>(AnimSection->MapTimeToAnimation(CurrentTime.Time, CurrentTime.Rate));
+											const double Seconds = AnimSection->MapTimeToAnimation(CurrentTime.Time, CurrentTime.Rate);
 											FAnimExtractContext ExtractionContext(Seconds, false);
 
 											AnimSequence->GetAnimationPose(OutAnimationPoseData, ExtractionContext);
diff --git a/Engine/Source/Editor/MovieSceneTools/Private/TrackEditors/TransformTrackEditor.cpp b/Engine/Source/Editor/MovieSceneTools/Private/TrackEditors/TransformTrackEditor.cpp
--- a/Engine/Source/Editor/MovieSceneTools/Private/TrackEditors/TransformTrackEditor.cpp
+++ b/Engine/Source/Editor/MovieSceneTools/Private/TrackEditors/TransformTrackEditor.cpp
@@ -27,7 +27,7 @@
 #include "TransformPropertySection.h"
 #include "SequencerUtilities.h"
 #include "MovieSceneToolHelpers.h"
-#include "Animation/AnimData/AnimDataModel.h"
+#include "Animation/AnimData/IAnimationDataModel.h"
 
 #include "EntitySystem/Interrogation/MovieSceneInterrogationLinker.h"
 #include "EntitySystem/Interrogation/MovieSceneInterrogatedPropertyInstantiator.h"
diff --git a/Engine/Source/Editor/Persona/Private/AnimTimeline/AnimModel_AnimMontage.cpp b/Engine/Source/Editor/Persona/Private/AnimTimeline/AnimModel_AnimMontage.cpp
--- a/Engine/Source/Editor/Persona/Private/AnimTimeline/AnimModel_AnimMontage.cpp
+++ b/Engine/Source/Editor/Persona/Private/AnimTimeline/AnimModel_AnimMontage.cpp
@@ -225,14 +225,14 @@ void FAnimModel_AnimMontage::RecalculateSequenceLength()
 	{
 		AnimMontage->InvalidateRecursiveAsset();
 
-		float NewSequenceLength = CalculateSequenceLengthOfEditorObject();
-		if (NewSequenceLength != AnimMontage->GetPlayLength())
+		const float CurrentCalculatedLength = CalculateSequenceLengthOfEditorObject();
+		if(!FMath::IsNearlyEqual(CurrentCalculatedLength, AnimMontage->GetPlayLength(), UE_KINDA_SMALL_NUMBER))
 		{
-			ClampToEndTime(NewSequenceLength);
+			ClampToEndTime(CurrentCalculatedLength);
 
 			RefreshSectionTimes();
 
-			AnimMontage->SetCompositeLength(NewSequenceLength);
+			AnimMontage->SetCompositeLength(CurrentCalculatedLength);
 
 			// Reset view if we changed length (note: has to be done after ->SetCompositeLength)!
 			UpdateRange();
@@ -312,7 +312,7 @@ void FAnimModel_AnimMontage::ToggleSectionTimingDisplay()
 	bSectionTimingEnabled = !bSectionTimingEnabled;
 }
 
-void FAnimModel_AnimMontage::OnDataModelChanged(const EAnimDataModelNotifyType& NotifyType, UAnimDataModel* Model, const FAnimDataModelNotifPayload& PayLoad)
+void FAnimModel_AnimMontage::OnDataModelChanged(const EAnimDataModelNotifyType& NotifyType, IAnimationDataModel* Model, const FAnimDataModelNotifPayload& PayLoad)
 {
 	NotifyCollector.Handle(NotifyType);
 
diff --git a/Engine/Source/Editor/Persona/Private/AnimTimeline/AnimModel_AnimMontage.h b/Engine/Source/Editor/Persona/Private/AnimTimeline/AnimModel_AnimMontage.h
--- a/Engine/Source/Editor/Persona/Private/AnimTimeline/AnimModel_AnimMontage.h
+++ b/Engine/Source/Editor/Persona/Private/AnimTimeline/AnimModel_AnimMontage.h
@@ -62,7 +62,7 @@ public:
 	void ToggleSectionTimingDisplay();
 
 protected:	
-	virtual void OnDataModelChanged(const EAnimDataModelNotifyType& NotifyType, UAnimDataModel* Model, const FAnimDataModelNotifPayload& PayLoad) override;
+	virtual void OnDataModelChanged(const EAnimDataModelNotifyType& NotifyType, IAnimationDataModel* Model, const FAnimDataModelNotifPayload& PayLoad) override;
 
 private:
 	/** The anim montage we wrap */
diff --git a/Engine/Source/Editor/Persona/Private/AnimTimeline/AnimModel_AnimSequenceBase.cpp b/Engine/Source/Editor/Persona/Private/AnimTimeline/AnimModel_AnimSequenceBase.cpp
--- a/Engine/Source/Editor/Persona/Private/AnimTimeline/AnimModel_AnimSequenceBase.cpp
+++ b/Engine/Source/Editor/Persona/Private/AnimTimeline/AnimModel_AnimSequenceBase.cpp
@@ -311,7 +311,7 @@ void FAnimModel_AnimSequenceBase::RefreshAttributeTracks()
 	}
 }
 
-void FAnimModel_AnimSequenceBase::OnDataModelChanged(const EAnimDataModelNotifyType& NotifyType, UAnimDataModel* Model, const FAnimDataModelNotifPayload& PayLoad)
+void FAnimModel_AnimSequenceBase::OnDataModelChanged(const EAnimDataModelNotifyType& NotifyType, IAnimationDataModel* Model, const FAnimDataModelNotifPayload& PayLoad)
 {
 	NotifyCollector.Handle(NotifyType);
 
diff --git a/Engine/Source/Editor/Persona/Private/AnimTimeline/AnimModel_AnimSequenceBase.h b/Engine/Source/Editor/Persona/Private/AnimTimeline/AnimModel_AnimSequenceBase.h
--- a/Engine/Source/Editor/Persona/Private/AnimTimeline/AnimModel_AnimSequenceBase.h
+++ b/Engine/Source/Editor/Persona/Private/AnimTimeline/AnimModel_AnimSequenceBase.h
@@ -59,8 +59,8 @@ protected:
 	/** Refresh attribute tracks */
 	void RefreshAttributeTracks();
 
-	/** Callback for any change made to the UAnimDataModel embedded in the AnimSequenceBase instance this represents */
-	virtual void OnDataModelChanged(const EAnimDataModelNotifyType& NotifyType, UAnimDataModel* Model, const FAnimDataModelNotifPayload& PayLoad);
+	/** Callback for any change made to the IAnimationDataModel embedded in the AnimSequenceBase instance this represents */
+	virtual void OnDataModelChanged(const EAnimDataModelNotifyType& NotifyType, IAnimationDataModel* Model, const FAnimDataModelNotifPayload& PayLoad);
 
 private:
 	/** UI handlers */
diff --git a/Engine/Source/Editor/Persona/Private/AnimTimeline/AnimTimelineTrack_FloatCurve.cpp b/Engine/Source/Editor/Persona/Private/AnimTimeline/AnimTimelineTrack_FloatCurve.cpp
--- a/Engine/Source/Editor/Persona/Private/AnimTimeline/AnimTimelineTrack_FloatCurve.cpp
+++ b/Engine/Source/Editor/Persona/Private/AnimTimeline/AnimTimelineTrack_FloatCurve.cpp
@@ -19,7 +19,7 @@
 #include "Widgets/Input/SComboButton.h"
 #include "Widgets/Layout/SBox.h"
 #include "AnimModel_AnimSequenceBase.h"
-#include "Animation/AnimData/AnimDataModel.h"
+#include "Animation/AnimData/IAnimationDataModel.h"
 #include "SAnimOutlinerItem.h"
 
 #define LOCTEXT_NAMESPACE "FAnimTimelineTrack_FloatCurve"
@@ -42,7 +42,8 @@ TSharedRef<SWidget> FAnimTimelineTrack_FloatCurve::MakeTimelineWidgetContainer()
 	// zoom to fit now we have a view
 	CurveEditor->ZoomToFit(EAxisList::Y);
 
-	auto ColorLambda = [this]()
+	FLinearColor CurveColor = FloatCurve->GetCurveTypeFlag(AACF_Metadata) ? FloatCurve->GetColor().Desaturate(0.25f) : FloatCurve->GetColor().Desaturate(0.75f);
+	auto ColorLambda = [this, CurveColor]()
 	{
 		if(GetModel()->IsTrackSelected(AsShared()))
 		{
@@ -50,14 +51,14 @@ TSharedRef<SWidget> FAnimTimelineTrack_FloatCurve::MakeTimelineWidgetContainer()
 		}
 		else
 		{
-			return FloatCurve->GetCurveTypeFlag(AACF_Metadata) ? FloatCurve->GetColor().Desaturate(0.25f) : FloatCurve->GetColor().Desaturate(0.75f);
+			return CurveColor;
 		}
 	};
 
 	return
 		SAssignNew(TimelineWidgetContainer, SBorder)
 		.Padding(0.0f)
-		.BorderImage_Lambda([this](){ return FloatCurve->GetCurveTypeFlag(AACF_Metadata) ? FAppStyle::GetBrush("Sequencer.Section.SelectedSectionOverlay") : FAppStyle::GetBrush("AnimTimeline.Outliner.DefaultBorder"); })
+		.BorderImage(FloatCurve->GetCurveTypeFlag(AACF_Metadata) ? FAppStyle::GetBrush("Sequencer.Section.SelectedSectionOverlay") : FAppStyle::GetBrush("AnimTimeline.Outliner.DefaultBorder"))
 		.BorderBackgroundColor_Lambda(ColorLambda)
 		[
 			CurveWidget
@@ -254,14 +255,13 @@ void FAnimTimelineTrack_FloatCurve::OnCommitCurveName(const FText& InText, EText
           
 			FAnimationCurveIdentifier NewCurveId(NewSmartName, ERawCurveTrackTypes::RCT_Float);
 			Controller.RenameCurve(CurveId, NewCurveId);
-			Controller.RemoveBoneTracksMissingFromSkeleton(AnimSequenceBase->GetSkeleton());           
 		}
 	}
 }
 
 FText FAnimTimelineTrack_FloatCurve::GetLabel() const
 {
-	return FAnimTimelineTrack_FloatCurve::GetFloatCurveName(GetModel(), FloatCurve->Name);
+	return FAnimTimelineTrack_FloatCurve::GetFloatCurveName(GetModel(), CurveName);
 }
 
 FText FAnimTimelineTrack_FloatCurve::GetFloatCurveName(const TSharedRef<FAnimModel>& InModel, const FSmartName& InSmartName)
@@ -303,11 +303,7 @@ void FAnimTimelineTrack_FloatCurve::AddCurveTrackButton(TSharedPtr<SHorizontalBo
 		PersonaUtils::MakeTrackButton(LOCTEXT("EditCurveButtonText", "Curve"), FOnGetContent::CreateSP(this, &FAnimTimelineTrack_FloatCurve::BuildCurveTrackMenu), MakeAttributeSP(this, &FAnimTimelineTrack_FloatCurve::IsHovered))
 	];
 
-	auto GetValue = [this]()
-	{
-		return FloatCurve->Color;
-	};
-
+	FLinearColor CurveColor = FloatCurve->Color;
 	auto SetValue = [this](FLinearColor InNewColor)
 	{
 		UAnimSequenceBase* AnimSequenceBase = GetModel()->GetAnimSequenceBase();
@@ -320,11 +316,11 @@ void FAnimTimelineTrack_FloatCurve::AddCurveTrackButton(TSharedPtr<SHorizontalBo
 		}
 	};
 
-	auto OnGetMenuContent = [this, GetValue, SetValue]()
+	auto OnGetMenuContent = [this, CurveColor, SetValue]()
 	{
 		// Open a color picker
 		return SNew(SColorPicker)
-			.TargetColorAttribute_Lambda(GetValue)
+			.TargetColorAttribute(CurveColor)
 			.UseAlpha(false)
 			.DisplayInlineVersion(true)
 			.OnColorCommitted_Lambda(SetValue)
@@ -357,7 +353,7 @@ void FAnimTimelineTrack_FloatCurve::AddCurveTrackButton(TSharedPtr<SHorizontalBo
 		.ButtonContent()
 		[
 			SNew(SColorBlock)
-			.Color_Lambda(GetValue)
+			.Color(CurveColor)
 			.ShowBackgroundForAlpha(false)
 			.AlphaDisplayMode(EColorBlockAlphaDisplayMode::Ignore)
 			.Size(FVector2D(OutlinerRightPadding - 2.0f, OutlinerRightPadding))
diff --git a/Engine/Source/Editor/Persona/Private/AnimationEditorViewportClient.cpp b/Engine/Source/Editor/Persona/Private/AnimationEditorViewportClient.cpp
--- a/Engine/Source/Editor/Persona/Private/AnimationEditorViewportClient.cpp
+++ b/Engine/Source/Editor/Persona/Private/AnimationEditorViewportClient.cpp
@@ -577,7 +577,7 @@ void FAnimationViewportClient::DrawUVsForMesh(FViewport* InViewport, FCanvas* In
 
 	TArray<FVector2D> SelectedEdgeTexCoords; //No functionality in Persona for this (yet?)
 
-	DrawUVs(InViewport, InCanvas, InTextYPos, LODLevel, UVChannelToDraw, SelectedEdgeTexCoords, NULL, &PreviewMeshComponent->GetSkeletalMeshRenderData()->LODRenderData[LODLevel] );
+	DrawUVs(InViewport, InCanvas, InTextYPos, LODLevel, UVChannelToDraw, SelectedEdgeTexCoords, nullptr, &PreviewMeshComponent->GetSkeletalMeshRenderData()->LODRenderData[LODLevel] );
 }
 
 void FAnimationViewportClient::Tick(float DeltaSeconds) 
@@ -801,7 +801,7 @@ bool FAnimationViewportClient::ShouldDisplayAdditiveScaleErrorMessage() const
 	{
 		if (AnimSequence->IsValidAdditive() && AnimSequence->RefPoseSeq)
 		{
-			FGuid AnimSeqGuid = AnimSequence->RefPoseSeq->GetRawDataGuid();
+			FGuid AnimSeqGuid = AnimSequence->RefPoseSeq->GetDataModel()->GenerateGuid();
 			if (RefPoseGuid != AnimSeqGuid)
 			{
 				RefPoseGuid = AnimSeqGuid;
@@ -927,7 +927,7 @@ FText FAnimationViewportClient::GetDisplayInfo(bool bDisplayAllInfo) const
 	{
 		if (!PreviewMeshComponent->CheckIfBoundsAreCorrrect())
 		{
-			if( PreviewMeshComponent->GetPhysicsAsset() == NULL )
+			if( PreviewMeshComponent->GetPhysicsAsset() == nullptr )
 			{
 				TextValue = ConcatenateLine(TextValue, LOCTEXT("NeedToSetupPhysicsAssetForAccurateBounds", "<AnimViewport.WarningText>You may need to setup Physics Asset to use more accurate bounds</>"));
 			}
@@ -938,7 +938,7 @@ FText FAnimationViewportClient::GetDisplayInfo(bool bDisplayAllInfo) const
 		}
 	}
 
-	if (PreviewMeshComponent != NULL && PreviewMeshComponent->MeshObject != NULL)
+	if (PreviewMeshComponent != nullptr && PreviewMeshComponent->MeshObject != nullptr)
 	{
 		if (bDisplayAllInfo)
 		{
@@ -1149,7 +1149,7 @@ FText FAnimationViewportClient::GetDisplayInfo(bool bDisplayAllInfo) const
 
 	if (const UPoseAsset* PoseAsset = Cast<UPoseAsset>(GetAnimPreviewScene()->GetPreviewAnimationAsset()))
 	{
-		if (PoseAsset->SourceAnimation && PoseAsset->SourceAnimation->GetRawDataGuid() != PoseAsset->SourceAnimationRawDataGUID)
+		if (PoseAsset->SourceAnimation && PoseAsset->SourceAnimation->GetDataModel()->GenerateGuid() != PoseAsset->SourceAnimationRawDataGUID)
 		{
 			TextValue = ConcatenateLine(TextValue, LOCTEXT("PoseAssetOutOfDateWarning", "<AnimViewport.WarningText>Poses are out-of-sync with the source animation. To update them click \"Update Source\"</>"));
 		}
@@ -2300,7 +2300,7 @@ void FAnimationViewportClient::SetupViewForRendering( FSceneViewFamily& ViewFami
 
 	// Cache screen size
 	UDebugSkelMeshComponent* PreviewMeshComponent = GetAnimPreviewScene()->GetPreviewMeshComponent();
-	if (PreviewMeshComponent != NULL && PreviewMeshComponent->MeshObject != NULL)
+	if (PreviewMeshComponent != nullptr && PreviewMeshComponent->MeshObject != nullptr)
 	{
 		const FBoxSphereBounds& SkelBounds = PreviewMeshComponent->Bounds;
 		CachedScreenSize = ComputeBoundsScreenSize(SkelBounds.Origin, SkelBounds.SphereRadius, View);
diff --git a/Engine/Source/Editor/Persona/Private/SAnimSequenceCurveEditor.cpp b/Engine/Source/Editor/Persona/Private/SAnimSequenceCurveEditor.cpp
--- a/Engine/Source/Editor/Persona/Private/SAnimSequenceCurveEditor.cpp
+++ b/Engine/Source/Editor/Persona/Private/SAnimSequenceCurveEditor.cpp
@@ -17,7 +17,7 @@
 #include "Editor.h"
 #include "Framework/MultiBox/MultiBoxBuilder.h"
 #include "Styling/AppStyle.h"
-#include "Animation/AnimData/AnimDataModel.h"
+#include "Animation/AnimData/IAnimationDataModel.h"
 
 #define LOCTEXT_NAMESPACE "SAnimSequenceCurveEditor"
 
@@ -122,7 +122,7 @@ void FRichCurveEditorModelNamed::CurveHasChanged()
 	}
 }
 
-void FRichCurveEditorModelNamed::OnModelHasChanged(const EAnimDataModelNotifyType& NotifyType, UAnimDataModel* Model, const FAnimDataModelNotifPayload& Payload)
+void FRichCurveEditorModelNamed::OnModelHasChanged(const EAnimDataModelNotifyType& NotifyType, IAnimationDataModel* Model, const FAnimDataModelNotifPayload& Payload)
 {
 	NotifyCollector.Handle(NotifyType);
 
@@ -367,7 +367,7 @@ SAnimSequenceCurveEditor::~SAnimSequenceCurveEditor()
 {
 	if(AnimSequence)
 	{
-		AnimSequence->GetDataModel()->GetModifiedEvent().RemoveAll(this);		
+		AnimSequence->GetDataModel()->GetModifiedEvent().RemoveAll(this);
 	}
 }
 
@@ -438,7 +438,7 @@ void SAnimSequenceCurveEditor::Construct(const FArguments& InArgs, const TShared
 	];
 }
 
-void SAnimSequenceCurveEditor::OnModelHasChanged(const EAnimDataModelNotifyType& NotifyType, UAnimDataModel* Model, const FAnimDataModelNotifPayload& Payload)
+void SAnimSequenceCurveEditor::OnModelHasChanged(const EAnimDataModelNotifyType& NotifyType, IAnimationDataModel* Model, const FAnimDataModelNotifPayload& Payload)
 {
 	auto StopEditingCurve = [this, NotifyType, &Payload, Model]()
 	{
diff --git a/Engine/Source/Editor/Persona/Private/SAnimSequenceCurveEditor.h b/Engine/Source/Editor/Persona/Private/SAnimSequenceCurveEditor.h
--- a/Engine/Source/Editor/Persona/Private/SAnimSequenceCurveEditor.h
+++ b/Engine/Source/Editor/Persona/Private/SAnimSequenceCurveEditor.h
@@ -35,7 +35,7 @@ public:
 	virtual void SetKeyAttributes(TArrayView<const FKeyHandle> InKeys, TArrayView<const FKeyAttributes> InAttributes, EPropertyChangeType::Type ChangeType = EPropertyChangeType::Unspecified) override;
 
 	void CurveHasChanged();
-	void OnModelHasChanged(const EAnimDataModelNotifyType& NotifyType, UAnimDataModel* Model, const FAnimDataModelNotifPayload& Payload);
+	void OnModelHasChanged(const EAnimDataModelNotifyType& NotifyType, IAnimationDataModel* Model, const FAnimDataModelNotifPayload& Payload);
 	void UpdateCachedCurve();
 
 	FSmartName Name;
@@ -72,8 +72,7 @@ class SAnimSequenceCurveEditor : public IAnimSequenceCurveEditor
 	virtual void RemoveCurve(const FSmartName& InName, ERawCurveTrackTypes InType, int32 InCurveIndex) override;
 	virtual void ZoomToFit() override;
 	
-	void OnModelHasChanged(const EAnimDataModelNotifyType& NotifyType, UAnimDataModel* Model, const FAnimDataModelNotifPayload& Payload);
-
+	void OnModelHasChanged(const EAnimDataModelNotifyType& NotifyType, IAnimationDataModel* Model, const FAnimDataModelNotifPayload& Payload);
 private:
 	// Build the toolbar for this curve editor
 	TSharedRef<SWidget> MakeToolbar(TSharedRef<SCurveEditorPanel> InEditorPanel);
diff --git a/Engine/Source/Editor/Persona/Private/SAnimationScrubPanel.cpp b/Engine/Source/Editor/Persona/Private/SAnimationScrubPanel.cpp
--- a/Engine/Source/Editor/Persona/Private/SAnimationScrubPanel.cpp
+++ b/Engine/Source/Editor/Persona/Private/SAnimationScrubPanel.cpp
@@ -10,7 +10,7 @@
 #include "ScopedTransaction.h"
 #include "Animation/BlendSpace.h"
 #include "AnimationEditorPreviewScene.h"
-#include "Animation/AnimData/AnimDataModel.h"
+#include "Animation/AnimData/IAnimationDataModel.h"
 #include "Animation/AnimSequenceHelpers.h"
 
 #define LOCTEXT_NAMESPACE "AnimationScrubPanel"
@@ -373,7 +373,7 @@ bool SAnimationScrubPanel::DoesSyncViewport() const
 {
 	UAnimSingleNodeInstance* PreviewInstance = GetPreviewInstance();
 
-	return (( LockedSequence==NULL && PreviewInstance ) || ( LockedSequence && PreviewInstance && PreviewInstance->GetCurrentAsset() == LockedSequence ));
+	return (( LockedSequence==nullptr && PreviewInstance ) || ( LockedSequence && PreviewInstance && PreviewInstance->GetCurrentAsset() == LockedSequence ));
 }
 
 void SAnimationScrubPanel::Tick( const FGeometry& AllottedGeometry, const double InCurrentTime, const float InDeltaTime )
@@ -387,7 +387,7 @@ void SAnimationScrubPanel::Tick( const FGeometry& AllottedGeometry, const double
 class UAnimSingleNodeInstance* SAnimationScrubPanel::GetPreviewInstance() const
 {
 	UDebugSkelMeshComponent* PreviewMeshComponent = GetPreviewScene()->GetPreviewMeshComponent();
-	return PreviewMeshComponent && PreviewMeshComponent->IsPreviewOn()? PreviewMeshComponent->PreviewInstance : NULL;
+	return PreviewMeshComponent && PreviewMeshComponent->IsPreviewOn()? PreviewMeshComponent->PreviewInstance : nullptr;
 }
 
 float SAnimationScrubPanel::GetScrubValue() const
@@ -424,20 +424,20 @@ UAnimInstance* SAnimationScrubPanel::GetAnimInstanceWithBlueprint() const
 	{
 		UAnimInstance* Instance = DebugComponent->GetAnimInstance();
 
-		if ((Instance != NULL) && (Instance->GetClass()->ClassGeneratedBy != NULL))
+		if ((Instance != nullptr) && (Instance->GetClass()->ClassGeneratedBy != nullptr))
 		{
 			return Instance;
 		}
 	}
 
-	return NULL;
+	return nullptr;
 }
 
 bool SAnimationScrubPanel::GetAnimBlueprintDebugData(UAnimInstance*& Instance, FAnimBlueprintDebugData*& DebugInfo) const
 {
 	Instance = GetAnimInstanceWithBlueprint();
 
-	if (Instance != NULL)
+	if (Instance != nullptr)
 	{
 		// Avoid updating the instance if we're replaying the past
 		if (UAnimBlueprintGeneratedClass* AnimBlueprintClass = Cast<UAnimBlueprintGeneratedClass>(Instance->GetClass()))
diff --git a/Engine/Source/Editor/Persona/Private/SAnimationSequenceBrowser.cpp b/Engine/Source/Editor/Persona/Private/SAnimationSequenceBrowser.cpp
--- a/Engine/Source/Editor/Persona/Private/SAnimationSequenceBrowser.cpp
+++ b/Engine/Source/Editor/Persona/Private/SAnimationSequenceBrowser.cpp
@@ -469,7 +469,7 @@ SAnimationSequenceBrowser::~SAnimationSequenceBrowser()
 
 	if(ViewportClient.IsValid())
 	{
-		ViewportClient->Viewport = NULL;
+		ViewportClient->Viewport = nullptr;
 	}
 }
 
@@ -700,7 +700,7 @@ void SAnimationSequenceBrowser::GetSelectedPackages(const TArray<FAssetData>& As
 {
 	for (int32 AssetIdx = 0; AssetIdx < Assets.Num(); ++AssetIdx)
 	{
-		UPackage* Package = FindPackage(NULL, *Assets[AssetIdx].PackageName.ToString());
+		UPackage* Package = FindPackage(nullptr, *Assets[AssetIdx].PackageName.ToString());
 
 		if ( Package )
 		{
@@ -1014,7 +1014,7 @@ void SAnimationSequenceBrowser::Construct(const FArguments& InArgs, const TShare
 	// Making Skeleton to be private, and now GET_MEMBER_NAME_CHECKED doesn't work
 	AssetRegistryTagsToIgnore.Add(TEXT("Skeleton"));
 	PRAGMA_DISABLE_DEPRECATION_WARNINGS
-	AssetRegistryTagsToIgnore.Add(GET_MEMBER_NAME_CHECKED(UAnimSequenceBase, SequenceLength));
+	AssetRegistryTagsToIgnore.Add(TEXT("SequenceLength"));
 	PRAGMA_ENABLE_DEPRECATION_WARNINGS
 	AssetRegistryTagsToIgnore.Add(GET_MEMBER_NAME_CHECKED(UAnimSequenceBase, RateScale));
 }
@@ -1066,7 +1066,7 @@ FReply SAnimationSequenceBrowser::OnMouseDownHistory( const FGeometry& MyGeometr
 
 TSharedRef<SWidget> SAnimationSequenceBrowser::CreateHistoryMenu(bool bInBackHistory) const
 {
-	FMenuBuilder MenuBuilder(true, NULL);
+	FMenuBuilder MenuBuilder(true, nullptr);
 	if(bInBackHistory)
 	{
 		int32 HistoryIdx = CurrentAssetHistoryIndex - 1;
diff --git a/Engine/Source/Editor/SequenceRecorder/Private/AnimationRecorder.cpp b/Engine/Source/Editor/SequenceRecorder/Private/AnimationRecorder.cpp
--- a/Engine/Source/Editor/SequenceRecorder/Private/AnimationRecorder.cpp
+++ b/Engine/Source/Editor/SequenceRecorder/Private/AnimationRecorder.cpp
@@ -244,10 +244,12 @@ void FAnimationRecorder::StartRecord(USkeletalMeshComponent* Component, UAnimSeq
 	LastFrame = 0;
 
 	IAnimationDataController& Controller = AnimationObject->GetController();
-	Controller.SetModel(AnimationObject->GetDataModel());
+	Controller.SetModel(AnimationObject->GetDataModelInterface());
 
 	Controller.OpenBracket(LOCTEXT("StartRecord_Bracket", "Starting Animation Recording"));
 
+	Controller.InitializeModel();
+
 	const bool bKeepNotifiesAndCurves = CVarKeepNotifyAndCurvesOnAnimationRecord->GetInt() == 0 ? false : true;
 	if (bKeepNotifiesAndCurves)
 	{
@@ -359,8 +361,8 @@ UAnimSequence* FAnimationRecorder::StopRecord(bool bShowMessage)
 		int32 NumKeys = LastFrame.Value + 1;
 
 		// can't use TimePassed. That is just total time that has been passed, not necessarily match with frame count
-		Controller.SetPlayLength( (NumKeys>1) ? RecordingRate.AsSeconds(LastFrame): RecordingRate.AsSeconds(1) );
 		Controller.SetFrameRate(RecordingRate);
+		Controller.SetNumberOfFrames( FMath::Max(LastFrame.Value,1) );
 
 		ProcessNotifies();
 
@@ -521,7 +523,7 @@ UAnimSequence* FAnimationRecorder::StopRecord(bool bShowMessage)
 			FSavePackageArgs SaveArgs;
 			SaveArgs.TopLevelFlags = RF_Standalone;
 			SaveArgs.SaveFlags = SAVE_NoError;
-			UPackage::SavePackage(Package, NULL, *PackageFileName, SaveArgs);
+			UPackage::SavePackage(Package, nullptr, *PackageFileName, SaveArgs);
 
 			ElapsedTime = FPlatformTime::Seconds() - StartTime;
 			UE_LOG(LogAnimation, Log, TEXT("Animation Recorder saved %s in %0.2f seconds"), *PackageName, ElapsedTime);
@@ -561,7 +563,7 @@ UAnimSequence* FAnimationRecorder::StopRecord(bool bShowMessage)
 			FAssetRegistryModule::AssetCreated(AnimationObject);
 		}
 
-		AnimationObject = NULL;
+		AnimationObject = nullptr;
 		PreviousSpacesBases.Empty();
 		PreviousAnimCurves.Empty();
 
diff --git a/Engine/Source/Editor/UnrealEd/Private/Animation/DebugSkelMeshComponent.cpp b/Engine/Source/Editor/UnrealEd/Private/Animation/DebugSkelMeshComponent.cpp
--- a/Engine/Source/Editor/UnrealEd/Private/Animation/DebugSkelMeshComponent.cpp
+++ b/Engine/Source/Editor/UnrealEd/Private/Animation/DebugSkelMeshComponent.cpp
@@ -31,7 +31,7 @@ UDebugSkelMeshComponent::UDebugSkelMeshComponent(const FObjectInitializer& Objec
 	: Super(ObjectInitializer)
 {
 	bDrawMesh = true;
-	PreviewInstance = NULL;
+	PreviewInstance = nullptr;
 	bDisplayRawAnimation = false;
 	bDisplayNonRetargetedPose = false;
 
@@ -417,9 +417,9 @@ bool UDebugSkelMeshComponent::IsTrackingAttachedLOD() const
 
 FPrimitiveSceneProxy* UDebugSkelMeshComponent::CreateSceneProxy()
 {
-	FDebugSkelMeshSceneProxy* Result = NULL;
+	FDebugSkelMeshSceneProxy* Result = nullptr;
 	ERHIFeatureLevel::Type SceneFeatureLevel = GetWorld()->FeatureLevel;
-	FSkeletalMeshRenderData* SkelMeshRenderData = SkeletalMesh ? SkeletalMesh->GetResourceForRendering() : NULL;
+	FSkeletalMeshRenderData* SkelMeshRenderData = SkeletalMesh ? SkeletalMesh->GetResourceForRendering() : nullptr;
 
 	// only create a scene proxy for rendering if
 	// properly initialized
@@ -441,7 +441,7 @@ bool UDebugSkelMeshComponent::ShouldRenderSelected() const
 
 bool UDebugSkelMeshComponent::IsPreviewOn() const
 {
-	return (PreviewInstance != NULL) && (PreviewInstance == AnimScriptInstance);
+	return (PreviewInstance != nullptr) && (PreviewInstance == AnimScriptInstance);
 }
 
 FString UDebugSkelMeshComponent::GetPreviewText() const
@@ -488,7 +488,7 @@ void UDebugSkelMeshComponent::InitAnim(bool bForceReinit)
 		if (!SkeletalMesh->GetSkeleton()->IsCompatible(PreviewInstance->GetCurrentAsset()->GetSkeleton()))
 		{
 			// if it doesn't match, just clear it
-			PreviewInstance->SetAnimationAsset(NULL);
+			PreviewInstance->SetAnimationAsset(nullptr);
 		}
 	}
 
@@ -501,8 +501,8 @@ void UDebugSkelMeshComponent::InitAnim(bool bForceReinit)
 
 	Super::InitAnim(bForceReinit);
 
-	// if PreviewInstance is NULL, create here once
-	if (PreviewInstance == NULL)
+	// if PreviewInstance is nullptr, create here once
+	if (PreviewInstance == nullptr)
 	{
 		PreviewInstance = NewObject<UAnimPreviewInstance>(this);
 		check(PreviewInstance);
@@ -513,7 +513,7 @@ void UDebugSkelMeshComponent::InitAnim(bool bForceReinit)
 
 	// if anim script instance is null because it's not playing a blueprint, set to PreviewInstnace by default
 	// that way if user would like to modify bones or do extra stuff, it will work
-	if (AnimScriptInstance == NULL)
+	if (AnimScriptInstance == nullptr)
 	{
 		AnimScriptInstance = PreviewInstance;
 		AnimScriptInstance->InitializeAnimation();
@@ -566,7 +566,7 @@ void UDebugSkelMeshComponent::EnablePreview(bool bEnable, UAnimationAsset* Previ
 		}
 		else if (IsPreviewOn())
 		{
-			if (PreviewInstance->GetCurrentAsset() == PreviewAsset || PreviewAsset == NULL)
+			if (PreviewInstance->GetCurrentAsset() == PreviewAsset || PreviewAsset == nullptr)
 			{
 				// now recover to saved AnimScriptInstance;
 				AnimScriptInstance = SavedAnimScriptInstance;
@@ -719,7 +719,7 @@ void UDebugSkelMeshComponent::GenSpaceBases(TArray<FTransform>& OutSpaceBases)
 void UDebugSkelMeshComponent::RefreshBoneTransforms(FActorComponentTickFunction* TickFunction)
 {
 	// Run regular update first so we get RequiredBones up to date.
-	Super::RefreshBoneTransforms(NULL); // Pass NULL so we force non threaded work
+	Super::RefreshBoneTransforms(nullptr); // Pass nullptr so we force non threaded work
 
 	// none of these code works if we don't have anim instance, so no reason to check it for every if
 	if (AnimScriptInstance && AnimScriptInstance->GetRequiredBones().IsValid())
@@ -805,7 +805,7 @@ void UDebugSkelMeshComponent::RefreshBoneTransforms(FActorComponentTickFunction*
 						AdditiveBasePose.SetBoneContainer(&BoneContainer);
 						
 						FAnimationPoseData AnimationPoseData(AdditiveBasePose, AdditiveCurve, AdditiveAttributes);
-						Sequence->GetAdditiveBasePose(AnimationPoseData, FAnimExtractContext(PreviewInstance->GetCurrentTime()));
+						Sequence->GetAdditiveBasePose(AnimationPoseData, FAnimExtractContext(static_cast<double>(PreviewInstance->GetCurrentTime())));
 						CSAdditiveBasePose.InitPose(AnimationPoseData.GetPose());
 					}
 
diff --git a/Engine/Source/Editor/UnrealEd/Private/Fbx/FbxAnimationExport.cpp b/Engine/Source/Editor/UnrealEd/Private/Fbx/FbxAnimationExport.cpp
--- a/Engine/Source/Editor/UnrealEd/Private/Fbx/FbxAnimationExport.cpp
+++ b/Engine/Source/Editor/UnrealEd/Private/Fbx/FbxAnimationExport.cpp
@@ -72,6 +72,14 @@ void FFbxExporter::ExportAnimSequenceToFbx(const UAnimSequence* AnimSeq,
 									 float AnimEndOffset,
 									 float AnimPlayRate,
 									 float StartTime)
+{
+	ExportAnimSequenceToFbx(AnimSeq, SkelMesh, BoneNodes, InAnimLayer,
+		AnimSeq->GetDataModel()->GetFrameRate().AsFrameTime(AnimStartOffset),
+		AnimSeq->GetDataModel()->GetFrameRate().AsFrameTime(AnimSeq->GetPlayLength() - AnimEndOffset),
+		AnimPlayRate, StartTime);
+}
+
+void FFbxExporter::ExportAnimSequenceToFbx(const UAnimSequence* AnimSeq, const USkeletalMesh* SkelMesh, TArray<FbxNode*>& BoneNodes, FbxAnimLayer* InAnimLayer, FFrameTime StartFrameTime, FFrameTime EndFrameTime, float FrameRateScale, float StartTime)
 {
 	// stack allocator for extracting curve
 	FMemMark Mark(FMemStack::Get());
@@ -115,7 +123,7 @@ void FFbxExporter::ExportAnimSequenceToFbx(const UAnimSequence* AnimSeq,
 		}
 	}
 
-	ExportCustomAnimCurvesToFbx(CustomCurveMap, AnimSeq, AnimStartOffset, AnimEndOffset, AnimPlayRate, StartTime);
+	ExportCustomAnimCurvesToFbx(CustomCurveMap, AnimSeq, StartFrameTime, EndFrameTime, FrameRateScale, StartTime);
 
 	TArray<const FAnimatedBoneAttribute*> CustomAttributes;
 
@@ -210,7 +218,7 @@ void FFbxExporter::ExportAnimSequenceToFbx(const UAnimSequence* AnimSeq,
 			Curve->KeyModifyBegin();
 		}
 
-		auto ExportLambda = [&](float AnimTime, FbxTime ExportTime, bool bLastKey) {
+		auto ExportLambda = [&](double AnimTime, FbxTime ExportTime, bool bLastKey) {
 			FTransform BoneAtom;
 			AnimSeq->GetBoneTransform(BoneAtom, BoneTrackIndex, AnimTime, true);
 			FbxAMatrix FbxMatrix = Converter.ConvertMatrix(BoneAtom.ToMatrixWithScale());
@@ -268,7 +276,7 @@ void FFbxExporter::ExportAnimSequenceToFbx(const UAnimSequence* AnimSeq,
 			}
 		};
 
-		IterateInsideAnimSequence(AnimSeq, AnimStartOffset, AnimEndOffset, AnimPlayRate, StartTime, ExportLambda);
+		IterateInsideAnimSequence(AnimSeq, StartFrameTime, EndFrameTime, FrameRateScale, StartTime, ExportLambda);
 
 		for (FbxAnimCurve* Curve : Curves)
 		{
@@ -290,6 +298,15 @@ void FFbxExporter::ExportAnimSequenceToFbx(const UAnimSequence* AnimSeq,
 
 void FFbxExporter::ExportCustomAnimCurvesToFbx(const TMap<FName, FbxAnimCurve*>& CustomCurves, const UAnimSequence* AnimSeq, 
 	float AnimStartOffset, float AnimEndOffset, float AnimPlayRate, float StartTime, float ValueScale)
+{
+	ExportCustomAnimCurvesToFbx(CustomCurves, AnimSeq,
+		AnimSeq->GetDataModel()->GetFrameRate().AsFrameTime(AnimStartOffset),
+		AnimSeq->GetDataModel()->GetFrameRate().AsFrameTime(AnimSeq->GetPlayLength() - AnimEndOffset),
+		AnimPlayRate, StartTime, ValueScale);
+}
+	
+void FFbxExporter::ExportCustomAnimCurvesToFbx(const TMap<FName, FbxAnimCurve*>& CustomCurves, const UAnimSequence* AnimSeq,
+	FFrameTime AnimStartOffset, FFrameTime AnimEndOffset, float FrameRateScale, float StartTime, float ValueScale /*= 1.f*/)
 {
 	// stack allocator for extracting curve
 	FMemMark Mark(FMemStack::Get());
@@ -339,7 +356,7 @@ void FFbxExporter::ExportCustomAnimCurvesToFbx(const TMap<FName, FbxAnimCurve*>&
 		}
 	};
 
-	IterateInsideAnimSequence(AnimSeq, AnimStartOffset, AnimEndOffset, AnimPlayRate, StartTime, ExportLambda);
+	IterateInsideAnimSequence(AnimSeq, AnimStartOffset, AnimEndOffset, FrameRateScale, StartTime, ExportLambda);
 
 	for (auto CustomCurve : CustomCurves)
 	{
@@ -349,10 +366,21 @@ void FFbxExporter::ExportCustomAnimCurvesToFbx(const TMap<FName, FbxAnimCurve*>&
 
 void FFbxExporter::IterateInsideAnimSequence(const UAnimSequence* AnimSeq, float AnimStartOffset, float AnimEndOffset, float AnimPlayRate, float StartTime, TFunctionRef<void(float, FbxTime, bool)> IterationLambda)
 {
-	float AnimTime = AnimStartOffset;
-	float AnimEndTime = (AnimSeq->GetPlayLength() - AnimEndOffset);
-	double TimePerKey = AnimSeq->GetDataModel()->GetFrameRate().AsInterval();
-	const float AnimTimeIncrement = TimePerKey * AnimPlayRate;
+	IterateInsideAnimSequence(AnimSeq,
+		AnimSeq->GetDataModel()->GetFrameRate().AsFrameTime(AnimStartOffset),
+		AnimSeq->GetDataModel()->GetFrameRate().AsFrameTime(AnimSeq->GetPlayLength() - AnimEndOffset),
+		AnimPlayRate,
+		AnimPlayRate,
+		[IterationLambda](double T, FbxTime FT, bool B) -> void
+		{
+			IterationLambda((float)T, FT, B);
+		});
+}
+	
+void FFbxExporter::IterateInsideAnimSequence(const UAnimSequence* AnimSeq, FFrameTime StartFrameTime, FFrameTime EndFrameTime, float FrameRateScale, float StartTime, TFunctionRef<void(double, FbxTime, bool)> IterationLambda)
+{
+	const double TimePerKey = AnimSeq->GetDataModel()->GetFrameRate().AsInterval();
+	const double AnimTimeIncrement = TimePerKey * FrameRateScale;
 	uint32 AnimFrameIndex = 0;
 
 	FbxTime ExportTime;
@@ -363,15 +391,16 @@ void FFbxExporter::IterateInsideAnimSequence(const UAnimSequence* AnimSeq, float
 
 	// Step through each frame and add custom curve data
 	bool bLastKey = false;
+	FFrameTime FrameTime = StartFrameTime;
+	const FFrameRate& FrameRate = AnimSeq->GetDataModel()->GetFrameRate();
 	while (!bLastKey)
 	{
-		bLastKey = (AnimTime + KINDA_SMALL_NUMBER) > AnimEndTime;
-
-		IterationLambda(AnimTime, ExportTime, bLastKey);
+		bLastKey = FrameTime > EndFrameTime;
+		IterationLambda(FrameRate.AsSeconds(FrameTime), ExportTime, bLastKey);
 
 		ExportTime += ExportTimeIncrement;
 		AnimFrameIndex++;
-		AnimTime = AnimStartOffset + ((float)AnimFrameIndex * AnimTimeIncrement);
+		FrameTime += FFrameTime::FromDecimal(FrameRateScale);
 	}
 }
 
@@ -426,9 +455,9 @@ void FFbxExporter::CorrectAnimTrackInterpolation( TArray<FbxNode*>& BoneNodes, F
 
 FbxNode* FFbxExporter::ExportAnimSequence( const UAnimSequence* AnimSeq, const USkeletalMesh* SkelMesh, bool bExportSkelMesh, const TCHAR* MeshName, FbxNode* ActorRootNode, const TArray<UMaterialInterface*>* OverrideMaterials /*= nullptr*/ )
 {
-	if( Scene == NULL || AnimSeq == NULL || SkelMesh == NULL )
+	if( Scene == nullptr || AnimSeq == nullptr || SkelMesh == nullptr )
 	{
- 		return NULL;
+ 		return nullptr;
 	}
 
 
@@ -454,10 +483,14 @@ FbxNode* FFbxExporter::ExportAnimSequence( const UAnimSequence* AnimSeq, const U
 			SkelMesh,
 			BoneNodes,
 			AnimLayer,
-			0.f,		// AnimStartOffset
-			0.f,		// AnimEndOffset
-			1.f,		// AnimPlayRate
-			0.f);		// StartTime
+			// Start frame to export
+			FFrameTime(0),
+			// Final frame to export
+			FFrameTime(AnimSeq->GetDataModel()->GetNumberOfFrames()),
+			// Frame rate scale
+			1.f,
+			// FBX StartTime
+			0.f);
 
 		CorrectAnimTrackInterpolation(BoneNodes, AnimLayer);
 	}
@@ -554,7 +587,7 @@ FbxNode* FFbxExporter::ExportAnimSequence( const UAnimSequence* AnimSeq, const U
 
 void FFbxExporter::ExportAnimSequencesAsSingle( USkeletalMesh* SkelMesh, const ASkeletalMeshActor* SkelMeshActor, const FString& ExportName, const TArray<UAnimSequence*>& AnimSeqList, const TArray<struct FAnimControlTrackKey>& TrackKeys )
 {
-	if (Scene == NULL || SkelMesh == NULL || AnimSeqList.Num() == 0 || AnimSeqList.Num() != TrackKeys.Num()) return;
+	if (Scene == nullptr || SkelMesh == nullptr || AnimSeqList.Num() == 0 || AnimSeqList.Num() != TrackKeys.Num()) return;
 
 	FbxNode* BaseNode = FbxNode::Create(Scene, Converter.ConvertToFbxString(ExportName));
 	Scene->GetRootNode()->AddChild(BaseNode);
@@ -588,8 +621,8 @@ void FFbxExporter::ExportAnimSequencesAsSingle( USkeletalMesh* SkelMesh, const A
 			SkelMesh,
 			BoneNodes,
 			AnimLayer,
-			TrackKey.AnimStartOffset,
-			TrackKey.AnimEndOffset,
+			AnimSeq->GetDataModel()->GetFrameRate().AsFrameTime(TrackKey.AnimStartOffset),
+			AnimSeq->GetDataModel()->GetFrameRate().AsFrameTime(AnimSeq->GetPlayLength() - TrackKey.AnimEndOffset),
 			TrackKey.AnimPlayRate,
 			CurrentStartTime);
 	}
@@ -613,7 +646,7 @@ void FFbxExporter::ExportAnimSequencesAsSingle( USkeletalMesh* SkelMesh, const A
  */
 void FFbxExporter::ExportMatineeGroup(class AMatineeActor* MatineeActor, USkeletalMeshComponent* SkeletalMeshComponent)
 {
-	if (Scene == NULL || MatineeActor == NULL || SkeletalMeshComponent == NULL || MatineeActor->MatineeData->InterpLength == 0)
+	if (Scene == nullptr || MatineeActor == nullptr || SkeletalMeshComponent == nullptr || MatineeActor->MatineeData->InterpLength == 0)
 	{
 		return;
 	}
diff --git a/Engine/Source/Editor/UnrealEd/Private/Fbx/FbxSkeletalMeshExport.cpp b/Engine/Source/Editor/UnrealEd/Private/Fbx/FbxSkeletalMeshExport.cpp
--- a/Engine/Source/Editor/UnrealEd/Private/Fbx/FbxSkeletalMeshExport.cpp
+++ b/Engine/Source/Editor/UnrealEd/Private/Fbx/FbxSkeletalMeshExport.cpp
@@ -29,7 +29,7 @@ FbxNode* FFbxExporter::CreateSkeleton(const USkeletalMesh* SkelMesh, TArray<FbxN
 
 	if(RefSkeleton.GetRawBoneNum() == 0)
 	{
-		return NULL;
+		return nullptr;
 	}
 
 	// Create a list of the nodes we create for each bone, so that children can 
@@ -116,13 +116,13 @@ FbxNode* FFbxExporter::CreateMesh(const USkeletalMesh* SkelMesh, const TCHAR* Me
 	const int32 VertexCount = SourceModel.GetNumNonClothingVertices();
 
 	// Verify the integrity of the mesh.
-	if (VertexCount == 0) return NULL;
+	if (VertexCount == 0) return nullptr;
 
 	// Copy all the vertex data from the various chunks to a single buffer.
 	// Makes the rest of the code in this function cleaner and easier to maintain.  
 	TArray<FSoftSkinVertex> Vertices;
 	SourceModel.GetNonClothVertices(Vertices);
-	if (Vertices.Num() != VertexCount) return NULL;
+	if (Vertices.Num() != VertexCount) return nullptr;
 
 	FbxMesh* Mesh = FbxMesh::Create(Scene, TCHAR_TO_UTF8(MeshName));
 
@@ -137,7 +137,7 @@ FbxNode* FFbxExporter::CreateMesh(const USkeletalMesh* SkelMesh, const TCHAR* Me
 
 	// Create Layer 0 to hold the normals
 	FbxLayer* LayerZero = Mesh->GetLayer(0);
-	if (LayerZero == NULL)
+	if (LayerZero == nullptr)
 	{
 		Mesh->CreateLayer();
 		LayerZero = Mesh->GetLayer(0);
@@ -169,7 +169,7 @@ FbxNode* FFbxExporter::CreateMesh(const USkeletalMesh* SkelMesh, const TCHAR* Me
 	for (int32 TexCoordSourceIndex = 0; TexCoordSourceIndex < TexCoordSourceCount; ++TexCoordSourceIndex)
 	{
 		FbxLayer* Layer = Mesh->GetLayer(TexCoordSourceIndex);
-		if (Layer == NULL)
+		if (Layer == nullptr)
 		{
 			Mesh->CreateLayer();
 			Layer = Mesh->GetLayer(TexCoordSourceIndex);
@@ -328,10 +328,10 @@ FbxNode* FFbxExporter::CreateMesh(const USkeletalMesh* SkelMesh, const TCHAR* Me
 		if(AnimSeq && BlendShapeCurvesMap.Num() > 0)
 		{
 			ExportCustomAnimCurvesToFbx(BlendShapeCurvesMap, AnimSeq,
-				0.f,		// AnimStartOffset
-				0.f,		// AnimEndOffset
-				1.f,		// AnimPlayRate
-				0.f,		// StartTime
+				FFrameTime(0),													// Start frame to export
+				FFrameTime(AnimSeq->GetDataModel()->GetNumberOfFrames()),		// Final frame to export
+				1.f,												// Frame rate scale
+				0.f,												// FBX StartTime
 				100.f);		// ValueScale, for some reason we need to scale BlendShape curves by a factor of 100.
 		}
 	}
@@ -339,8 +339,6 @@ FbxNode* FFbxExporter::CreateMesh(const USkeletalMesh* SkelMesh, const TCHAR* Me
 	FbxNode* MeshNode = FbxNode::Create(Scene, TCHAR_TO_UTF8(MeshName));
 	MeshNode->SetNodeAttribute(Mesh);
 
-
-
 	// Add the materials for the mesh
 	const TArray<FSkeletalMaterial>& SkelMeshMaterials = SkelMesh->GetMaterials();
 	int32 MaterialCount = SkelMeshMaterials.Num();
@@ -358,7 +356,7 @@ FbxNode* FFbxExporter::CreateMesh(const USkeletalMesh* SkelMesh, const TCHAR* Me
 			MatInterface = SkelMeshMaterials[MaterialIndex].MaterialInterface;
 		}
 
-		FbxSurfaceMaterial* FbxMaterial = NULL;
+		FbxSurfaceMaterial* FbxMaterial = nullptr;
 		if (LODIndex == 0)
 		{
 			if (MatInterface && !FbxMaterials.Find(MatInterface))
@@ -590,7 +588,7 @@ void FFbxExporter::ExportSkeletalMeshComponent(USkeletalMeshComponent* SkelMeshC
 	if (SkelMeshComp && SkelMeshComp->SkeletalMesh)
 	{
 		UAnimSequence* AnimSeq = (bSaveAnimSeq && SkelMeshComp->GetAnimationMode() == EAnimationMode::AnimationSingleNode) ? 
-			Cast<UAnimSequence>(SkelMeshComp->AnimationData.AnimToPlay) : NULL;
+			Cast<UAnimSequence>(SkelMeshComp->AnimationData.AnimToPlay) : nullptr;
 		FbxNode* SkeletonRootNode = ExportSkeletalMeshToFbx(SkelMeshComp->SkeletalMesh, AnimSeq, MeshName, ActorRootNode, &ToRawPtrTArrayUnsafe(SkelMeshComp->OverrideMaterials));
 		if(SkeletonRootNode)
 		{
@@ -766,7 +764,7 @@ FbxNode* FFbxExporter::ExportSkeletalMeshToFbx(const USkeletalMesh* SkeletalMesh
 		return SkeletonRootNode;
 	}
 
-	return NULL;
+	return nullptr;
 }
 
 } // namespace UnFbx
diff --git a/Engine/Source/Editor/UnrealEd/Private/FbxExporter.h b/Engine/Source/Editor/UnrealEd/Private/FbxExporter.h
--- a/Engine/Source/Editor/UnrealEd/Private/FbxExporter.h
+++ b/Engine/Source/Editor/UnrealEd/Private/FbxExporter.h
@@ -217,7 +217,7 @@ public:
 	 * @param StaticMesh	The static mesh to export
 	 * @param MaterialOrder	Optional ordering of materials to set up correct material ID's across multiple meshes being export such as BSP surfaces which share common materials. Should be used sparingly
 	 */
-	virtual void ExportStaticMesh( UStaticMesh* StaticMesh, const TArray<FStaticMaterial>* MaterialOrder = NULL );
+	virtual void ExportStaticMesh( UStaticMesh* StaticMesh, const TArray<FStaticMaterial>* MaterialOrder = nullptr );
 
 	/**
 	 * Exports BSP
@@ -249,7 +249,7 @@ public:
 	/**
 	 * Exports a single UAnimSequence, and optionally a skeletal mesh
 	 */
-	FbxNode* ExportAnimSequence( const UAnimSequence* AnimSeq, const USkeletalMesh* SkelMesh, bool bExportSkelMesh, const TCHAR* MeshNames=NULL, FbxNode* ActorRootNode=NULL, const TArray<UMaterialInterface*>* OverrideMaterials = nullptr);
+	FbxNode* ExportAnimSequence( const UAnimSequence* AnimSeq, const USkeletalMesh* SkelMesh, bool bExportSkelMesh, const TCHAR* MeshNames=nullptr, FbxNode* ActorRootNode=nullptr, const TArray<UMaterialInterface*>* OverrideMaterials = nullptr);
 
 	/**
 	 * Exports the list of UAnimSequences as a single animation based on the settings in the TrackKeys
@@ -339,7 +339,7 @@ private:
 	 * @param MaterialOrderOverride	Optional ordering of materials to set up correct material ID's across multiple meshes being export such as BSP surfaces which share common materials. Should be used sparingly
 	 * @param OverrideMaterials	Optional array of materials to be used instead of the static mesh materials. Used for material overrides in static mesh components.
 	 */
-	FbxNode* ExportStaticMeshToFbx(const UStaticMesh* StaticMesh, int32 ExportLOD, const TCHAR* MeshName, FbxNode* FbxActor, int32 LightmapUVChannel = -1, const FColorVertexBuffer* ColorBuffer = NULL, const TArray<FStaticMaterial>* MaterialOrderOverride = NULL, const TArray<UMaterialInterface*>* OverrideMaterials = NULL);
+	FbxNode* ExportStaticMeshToFbx(const UStaticMesh* StaticMesh, int32 ExportLOD, const TCHAR* MeshName, FbxNode* FbxActor, int32 LightmapUVChannel = -1, const FColorVertexBuffer* ColorBuffer = nullptr, const TArray<FStaticMaterial>* MaterialOrderOverride = nullptr, const TArray<UMaterialInterface*>* OverrideMaterials = nullptr);
 
 	/**
 	 * Exports a spline mesh
@@ -421,18 +421,28 @@ private:
 	/**
 	 * Add the given animation sequence as rotation and translation tracks to the given list of bone nodes
 	 */
+	void ExportAnimSequenceToFbx(const UAnimSequence* AnimSeq, const USkeletalMesh* SkelMesh, TArray<FbxNode*>& BoneNodes, FbxAnimLayer* AnimLayer, FFrameTime StartFrameTime, FFrameTime EndFrameTime, float FrameRateScale, float StartTime);
+
+	UE_DEPRECATED(5.1, "ExportAnimSequenceToFbx is deprecated, use different signature")
 	void ExportAnimSequenceToFbx(const UAnimSequence* AnimSeq, const USkeletalMesh* SkelMesh, TArray<FbxNode*>& BoneNodes, FbxAnimLayer* AnimLayer,
 		float AnimStartOffset, float AnimEndOffset, float AnimPlayRate, float StartTime);
 
 	/**
 	 * Add the custom Curve data to the FbxAnimCurves passed in parameter by matching their name to the skeletal mesh custom curves.
 	 */
+	void ExportCustomAnimCurvesToFbx(const TMap<FName, FbxAnimCurve*>& CustomCurves, const UAnimSequence* AnimSeq,
+		FFrameTime AnimStartOffset, FFrameTime AnimEndOffset, float FrameRateScale, float StartTime, float ValueScale = 1.f);
+	
+	UE_DEPRECATED(5.1, "ExportCustomAnimCurvesToFbx is deprecated, use different signature")
 	void ExportCustomAnimCurvesToFbx(const TMap<FName, FbxAnimCurve*>& CustomCurves, const UAnimSequence* AnimSeq,
 		float AnimStartOffset, float AnimEndOffset, float AnimPlayRate, float StartTime, float ValueScale = 1.f);
 
 	/**
 	 * Used internally to reuse the AnimSequence iteration code when exporting various kind of curves.
 	 */
+	void IterateInsideAnimSequence(const UAnimSequence* AnimSeq, FFrameTime StartFrameTime, FFrameTime EndFrameTime, float FrameRateScale, float StartTime, TFunctionRef<void(double, FbxTime, bool)> IterationLambda);
+
+	UE_DEPRECATED(5.1, "IterateInsideAnimSequence is deprecated, use different signature")
 	void IterateInsideAnimSequence(const UAnimSequence* AnimSeq, float AnimStartOffset, float AnimEndOffset, float AnimPlayRate, float StartTime, TFunctionRef<void(float, FbxTime, bool)> IterationLambda);
 
 	/** 
diff --git a/Engine/Source/Editor/UnrealEd/Private/SkeletalMeshEdit.cpp b/Engine/Source/Editor/UnrealEd/Private/SkeletalMeshEdit.cpp
--- a/Engine/Source/Editor/UnrealEd/Private/SkeletalMeshEdit.cpp
+++ b/Engine/Source/Editor/UnrealEd/Private/SkeletalMeshEdit.cpp
@@ -49,7 +49,7 @@ UAnimSequence * UEditorEngine::ImportFbxAnimation( USkeleton* Skeleton, UObject*
 {
 	check(Skeleton);
 
-	UAnimSequence * NewAnimation=NULL;
+	UAnimSequence * NewAnimation=nullptr;
 
 	UnFbx::FFbxImporter* FFbxImporter = UnFbx::FFbxImporter::GetInstance();
 
@@ -76,7 +76,7 @@ UAnimSequence * UEditorEngine::ImportFbxAnimation( USkeleton* Skeleton, UObject*
 			FFbxImporter->AddTokenizedErrorMessage(FTokenizedMessage::Create(EMessageSeverity::Error, FText::Format(LOCTEXT("Error_CouldNotFindFbxTrack", "Mesh contains {0} bone as root but animation doesn't contain the root track.\nImport failed."), FText::FromName(Skeleton->GetReferenceSkeleton().GetBoneName(0)))), FFbxErrors::Animation_CouldNotFindRootTrack);
 
 			FFbxImporter->ReleaseScene();
-			return NULL;
+			return nullptr;
 		}
 
 		// Check for blend shape curves that are not skinned.  Unskinned geometry can still contain morph curves
@@ -568,19 +568,19 @@ void UnFbx::FFbxImporter::GetAnimationIntervalMultiLayer(FbxNode* RootNode, FbxA
 UAnimSequence * UnFbx::FFbxImporter::ImportAnimations(USkeleton* Skeleton, UObject* Outer, TArray<FbxNode*>& SortedLinks, const FString& Name, UFbxAnimSequenceImportData* TemplateImportData, TArray<FbxNode*>& NodeArray)
 {
 	// we need skeleton to create animsequence
-	if (Skeleton == NULL)
+	if (Skeleton == nullptr)
 	{
-		return NULL;
+		return nullptr;
 	}
 
 	int32 ValidTakeCount = 0;
 	if (IsValidAnimationData(SortedLinks, NodeArray, ValidTakeCount) == false)
 	{
 		AddTokenizedErrorMessage(FTokenizedMessage::Create(EMessageSeverity::Warning, LOCTEXT("FBXImport_InvalidAnimationData", "This does not contain any valid animation takes.")), FFbxErrors::Animation_InvalidData);
-		return NULL;
+		return nullptr;
 	}
 
-	UAnimSequence* LastCreatedAnim = NULL;
+	UAnimSequence* LastCreatedAnim = nullptr;
 
 	int32 ResampleRate = DEFAULT_SAMPLERATE;
 	if ( ImportOptions->bResample )
@@ -635,7 +635,7 @@ UAnimSequence * UnFbx::FFbxImporter::ImportAnimations(USkeleton* Skeleton, UObje
 
 		FString 	ParentPath = FString::Printf(TEXT("%s/%s"), *FPackageName::GetLongPackagePath(*Outer->GetName()), *SequenceName);
 		UObject* 	ParentPackage = CreatePackage( *ParentPath);
-		UObject* Object = LoadObject<UObject>(ParentPackage, *SequenceName, NULL, (LOAD_Quiet | LOAD_NoWarn), NULL);
+		UObject* Object = LoadObject<UObject>(ParentPackage, *SequenceName, nullptr, (LOAD_Quiet | LOAD_NoWarn), nullptr);
 		UAnimSequence * DestSeq = Cast<UAnimSequence>(Object);
 		// if object with same name exists, warn user
 		if (Object && !DestSeq)
@@ -645,12 +645,12 @@ UAnimSequence * UnFbx::FFbxImporter::ImportAnimations(USkeleton* Skeleton, UObje
 		}
 
 		// If not, create new one now.
-		if(!DestSeq)
+
+		const bool bCreateAsset = !DestSeq;
+		if(bCreateAsset)
 		{
 			DestSeq = NewObject<UAnimSequence>(ParentPackage, *SequenceName, RF_Public | RF_Standalone);
 			CreatedObjects.Add(DestSeq);
-			// Notify the asset registry
-			FAssetRegistryModule::AssetCreated(DestSeq);
 		}
 
 		DestSeq->SetSkeleton(Skeleton);
@@ -662,8 +662,16 @@ UAnimSequence * UnFbx::FFbxImporter::ImportAnimations(USkeleton* Skeleton, UObje
 		DestSeq->ImportFileFramerate = GetOriginalFbxFramerate();
 		DestSeq->ImportResampleFramerate = ResampleRate;
 
+		DestSeq->GetController().InitializeModel();
+
 		ImportAnimation(Skeleton, DestSeq, Name, SortedLinks, NodeArray, CurAnimStack, ResampleRate, AnimTimeSpan);
 
+		if (bCreateAsset)
+		{			
+			// Notify the asset registry
+			FAssetRegistryModule::AssetCreated(DestSeq);
+		}
+
 		LastCreatedAnim = DestSeq;
 	}
 
@@ -1337,7 +1345,7 @@ bool UnFbx::FFbxImporter::ImportCurveToAnimSequence(class UAnimSequence * Target
 		FAnimationCurveIdentifier FloatCurveId(NewName, ERawCurveTrackTypes::RCT_Float);
 
 
-		UAnimDataModel* DataModel = TargetSequence->GetDataModel();
+		IAnimationDataModel* DataModel = TargetSequence->GetDataModel();
 		IAnimationDataController& Controller = TargetSequence->GetController();
 
 		const FFloatCurve* TargetCurve = DataModel->FindFloatCurve(FloatCurveId);
@@ -1557,8 +1565,14 @@ bool UnFbx::FFbxImporter::ImportAnimation(USkeleton* Skeleton, UAnimSequence * D
 	//This destroy all previously imported animation raw data
 	Controller.RemoveAllBoneTracks();
 
+	// First set frame rate
+	const FFrameRate ResampleFrameRate(ResampleRate, 1);
+	Controller.SetFrameRate(ResampleFrameRate);
+
 	// if you have one pose(thus 0.f duration), it still contains animation, so we'll need to consider that as MINIMUM_ANIMATION_LENGTH time length
-	Controller.SetPlayLength(FGenericPlatformMath::Max<float>(SequenceLength.GetSecondDouble(), MINIMUM_ANIMATION_LENGTH));
+	const FFrameNumber NumberOfFrames = ResampleFrameRate.AsFrameNumber(SequenceLength.GetSecondDouble());
+	Controller.SetNumberOfFrames(FGenericPlatformMath::Max<int32>(NumberOfFrames.Value, 1));
+
 	if(PreviousSequenceLength > MINIMUM_ANIMATION_LENGTH && DestSeq->GetDataModel()->GetNumberOfFloatCurves() > 0)
 	{
 		// The sequence already existed when we began the import. We need to scale the key times for all curves to match the new 
@@ -1788,7 +1802,7 @@ void UnFbx::FFbxImporter::ImportAnimationCustomAttribute(FAnimCurveImportSetting
 	UAnimSequence* DestSeq = AnimImportSettings.DestSeq;
 	USkeleton* MySkeleton = DestSeq->GetSkeleton();
 
-	const UAnimDataModel* DataModel = DestSeq->GetDataModel();
+	const IAnimationDataModel* DataModel = DestSeq->GetDataModel();
 	const int32 NumFloatCurves = DataModel->GetNumberOfFloatCurves();
 	const FAnimationCurveData& CurveData = DataModel->GetCurveData();
 
@@ -1904,7 +1918,7 @@ void UnFbx::FFbxImporter::ImportBoneTracks(USkeleton* Skeleton, FAnimCurveImport
 	const FbxTime TimeComparisonThreshold = (KINDA_SMALL_NUMBER * static_cast<float>(FBXSDK_TC_SECOND));
 	const FReferenceSkeleton& RefSkeleton = Skeleton->GetReferenceSkeleton();
 
-	UAnimDataModel* DataModel = DestSeq->GetDataModel();
+	IAnimationDataModel* DataModel = DestSeq->GetDataModel();
 	IAnimationDataController& Controller = DestSeq->GetController();	
 	Controller.SetFrameRate(FFrameRate(ResampleRate, 1));
 
@@ -2065,7 +2079,7 @@ void UnFbx::FFbxImporter::ImportBoneTracks(USkeleton* Skeleton, FAnimCurveImport
 								if (AttributeIdentifier.IsValid())
 								{
 									// remove any existing attribute with the same identifier
-									if (const UAnimDataModel* Model = Controller.GetModel())
+									if (const IAnimationDataModel* Model = Controller.GetModel())
 									{
 										if (Model->FindAttribute(AttributeIdentifier))
 										{
diff --git a/Engine/Source/Editor/UnrealEd/Public/ChangeTransactor.h b/Engine/Source/Editor/UnrealEd/Public/ChangeTransactor.h
--- a/Engine/Source/Editor/UnrealEd/Public/ChangeTransactor.h
+++ b/Engine/Source/Editor/UnrealEd/Public/ChangeTransactor.h
@@ -42,7 +42,7 @@ struct FChangeTransactor
 		CheckTransactionObject();
 		checkf(!IsTransactionPending(), TEXT("A transaction is already pending"));
 		
-		if (GEngine && GEngine->CanTransact() && ensure(!GIsTransacting))
+		if (CanTransactChanges())
 		{
 			CompoundChangeData = MakeUnique<FCompoundChangeInput>();
 			PendingTransactionIndex = GEngine->BeginTransaction(TEXT("FChangeTransactor"), TransactionDescription, TransactionObject.Get());
@@ -100,7 +100,7 @@ struct FChangeTransactor
 
 		CheckTransactionObject();
 
-		if (GEngine && GEngine->CanTransact() && ensure(!GIsTransacting))
+		if (CanTransactChanges())
 		{
 			checkf(IsTransactionPending(), TEXT("No transaction was previously opened"));
 
@@ -118,6 +118,11 @@ struct FChangeTransactor
 			}
 		}
 	}
+	
+	static bool CanTransactChanges()
+	{
+		return GEngine && GEngine->CanTransact() && !GIsTransacting;
+	};
 private:	
 	void CheckTransactionObject()
 	{
@@ -136,5 +141,28 @@ private:
 	int32 TransactedChanges;
 };
 
+struct FScopedCompoundTransaction
+{
+	FScopedCompoundTransaction(FChangeTransactor& InTransactor, const FText& InDescription) : Transactor(InTransactor), bCreated(false)
+	{
+		if (FChangeTransactor::CanTransactChanges() && !Transactor.IsTransactionPending())
+		{
+			Transactor.OpenTransaction(InDescription);
+			bCreated = true;
+		}
+	}
+
+	~FScopedCompoundTransaction()
+	{
+		if (bCreated)
+		{
+			Transactor.CloseTransaction();
+		}
+	}
+
+	FChangeTransactor& Transactor;
+	bool bCreated;
+};
+
 }
 #endif // WITH_EDITOR
diff --git a/Engine/Source/Runtime/AnimGraphRuntime/Private/AnimNodes/AnimNode_BlendSpacePlayer.cpp b/Engine/Source/Runtime/AnimGraphRuntime/Private/AnimNodes/AnimNode_BlendSpacePlayer.cpp
--- a/Engine/Source/Runtime/AnimGraphRuntime/Private/AnimNodes/AnimNode_BlendSpacePlayer.cpp
+++ b/Engine/Source/Runtime/AnimGraphRuntime/Private/AnimNodes/AnimNode_BlendSpacePlayer.cpp
@@ -125,7 +125,7 @@ void FAnimNode_BlendSpacePlayerBase::Evaluate_AnyThread(FPoseContext& Output)
 	if ((CurrentBlendSpace != nullptr) && (Output.AnimInstanceProxy->IsSkeletonCompatible(CurrentBlendSpace->GetSkeleton())))
 	{
 		FAnimationPoseData AnimationPoseData(Output);
-		CurrentBlendSpace->GetAnimationPose(BlendSampleDataCache, FAnimExtractContext(InternalTimeAccumulator, Output.AnimInstanceProxy->ShouldExtractRootMotion(), DeltaTimeRecord, GetLoop()), AnimationPoseData);
+		CurrentBlendSpace->GetAnimationPose(BlendSampleDataCache, FAnimExtractContext(static_cast<double>(InternalTimeAccumulator), Output.AnimInstanceProxy->ShouldExtractRootMotion(), DeltaTimeRecord, GetLoop()), AnimationPoseData);
 	}
 	else
 	{
diff --git a/Engine/Source/Runtime/AnimGraphRuntime/Private/AnimNodes/AnimNode_RandomPlayer.cpp b/Engine/Source/Runtime/AnimGraphRuntime/Private/AnimNodes/AnimNode_RandomPlayer.cpp
--- a/Engine/Source/Runtime/AnimGraphRuntime/Private/AnimNodes/AnimNode_RandomPlayer.cpp
+++ b/Engine/Source/Runtime/AnimGraphRuntime/Private/AnimNodes/AnimNode_RandomPlayer.cpp
@@ -296,8 +296,8 @@ void FAnimNode_RandomPlayer::Evaluate_AnyThread(FPoseContext& Output)
 		FAnimationPoseData CurrentPoseData(Poses[0], Curves[0], Attributes[0]);
 		FAnimationPoseData NextPoseData(Poses[1], Curves[1], Attributes[1]);
 
-		CurrentSequence->GetAnimationPose(CurrentPoseData, FAnimExtractContext(CurrentData.CurrentPlayTime, AnimProxy->ShouldExtractRootMotion(), CurrentData.DeltaTimeRecord, CurrentData.RemainingLoops > 0));
-		NextSequence->GetAnimationPose(NextPoseData, FAnimExtractContext(NextData.CurrentPlayTime, AnimProxy->ShouldExtractRootMotion(), NextData.DeltaTimeRecord, NextData.RemainingLoops > 0));
+		CurrentSequence->GetAnimationPose(CurrentPoseData, FAnimExtractContext(static_cast<double>(CurrentData.CurrentPlayTime), AnimProxy->ShouldExtractRootMotion(), CurrentData.DeltaTimeRecord, CurrentData.RemainingLoops > 0));
+		NextSequence->GetAnimationPose(NextPoseData, FAnimExtractContext(static_cast<double>(NextData.CurrentPlayTime), AnimProxy->ShouldExtractRootMotion(), NextData.DeltaTimeRecord, NextData.RemainingLoops > 0));
 
 		FAnimationPoseData AnimationPoseData(Output);
 		FAnimationRuntime::BlendPosesTogether(Poses, Curves, Attributes, Weights, AnimationPoseData);
@@ -306,7 +306,7 @@ void FAnimNode_RandomPlayer::Evaluate_AnyThread(FPoseContext& Output)
 	{
 		// Single animation, no blending needed.
 		FAnimationPoseData AnimationPoseData(Output);
-		CurrentSequence->GetAnimationPose(AnimationPoseData, FAnimExtractContext(CurrentData.CurrentPlayTime, Output.AnimInstanceProxy->ShouldExtractRootMotion(), CurrentData.DeltaTimeRecord, CurrentData.RemainingLoops > 0));
+		CurrentSequence->GetAnimationPose(AnimationPoseData, FAnimExtractContext(static_cast<double>(CurrentData.CurrentPlayTime), Output.AnimInstanceProxy->ShouldExtractRootMotion(), CurrentData.DeltaTimeRecord, CurrentData.RemainingLoops > 0));
 	}
 }
 
diff --git a/Engine/Source/Runtime/AnimGraphRuntime/Private/AnimNodes/AnimNode_SequenceEvaluator.cpp b/Engine/Source/Runtime/AnimGraphRuntime/Private/AnimNodes/AnimNode_SequenceEvaluator.cpp
--- a/Engine/Source/Runtime/AnimGraphRuntime/Private/AnimNodes/AnimNode_SequenceEvaluator.cpp
+++ b/Engine/Source/Runtime/AnimGraphRuntime/Private/AnimNodes/AnimNode_SequenceEvaluator.cpp
@@ -89,7 +89,7 @@ void FAnimNode_SequenceEvaluatorBase::Evaluate_AnyThread(FPoseContext& Output)
 	if ((CurrentSequence != nullptr) && (Output.AnimInstanceProxy->IsSkeletonCompatible(CurrentSequence->GetSkeleton())))
 	{
 		FAnimationPoseData AnimationPoseData(Output);
-		CurrentSequence->GetAnimationPose(AnimationPoseData, FAnimExtractContext(InternalTimeAccumulator, Output.AnimInstanceProxy->ShouldExtractRootMotion(), DeltaTimeRecord, GetShouldLoop()));
+		CurrentSequence->GetAnimationPose(AnimationPoseData, FAnimExtractContext(static_cast<double>(InternalTimeAccumulator), Output.AnimInstanceProxy->ShouldExtractRootMotion(), DeltaTimeRecord, GetShouldLoop()));
 	}
 	else
 	{
diff --git a/Engine/Source/Runtime/Core/Public/UObject/UE5MainStreamObjectVersion.h b/Engine/Source/Runtime/Core/Public/UObject/UE5MainStreamObjectVersion.h
--- a/Engine/Source/Runtime/Core/Public/UObject/UE5MainStreamObjectVersion.h
+++ b/Engine/Source/Runtime/Core/Public/UObject/UE5MainStreamObjectVersion.h
@@ -239,6 +239,8 @@ struct CORE_API FUE5MainStreamObjectVersion
 		// Remove serialization of bounds relevant from  WorldPartitionActorDesc
 		WorldPartitionActorDescRemoveBoundsRelevantSerialization,
 
+		// Add IAnimationDataModel interface and replace UObject based representation for Animation Assets
+		AnimationDataModelInterface,
 		// -----<new versions can be added above this line>-------------------------------------------------
 		VersionPlusOne,
 		LatestVersion = VersionPlusOne - 1
diff --git a/Engine/Source/Runtime/Engine/Classes/Animation/AnimCompositeBase.h b/Engine/Source/Runtime/Engine/Classes/Animation/AnimCompositeBase.h
--- a/Engine/Source/Runtime/Engine/Classes/Animation/AnimCompositeBase.h
+++ b/Engine/Source/Runtime/Engine/Classes/Animation/AnimCompositeBase.h
@@ -39,7 +39,7 @@ struct FRootMotionExtractionStep
 	float EndPosition;
 
 	FRootMotionExtractionStep() 
-		: AnimSequence(NULL)
+		: AnimSequence(nullptr)
 		, StartPosition(0.f)
 		, EndPosition(0.f)
 		{
@@ -348,6 +348,10 @@ class UAnimCompositeBase : public UAnimSequenceBase
 	//~ Begin UObject Interface
 	virtual void PostLoad() override;
 	//~ End UObject Interface
+	
+	//~ Begin UAnimSequenceBase Interface
+	virtual FFrameRate GetSamplingFrameRate() const override;
+	//~ End UAnimSequenceBase Interface
 
 	// Extracts root motion from the supplied FAnimTrack between the Start End range specified
 	ENGINE_API void ExtractRootMotionFromTrack(const FAnimTrack &SlotAnimTrack, float StartTrackPosition, float EndTrackPosition, FRootMotionMovementParams &RootMotion) const;
@@ -362,5 +366,9 @@ class UAnimCompositeBase : public UAnimSequenceBase
 	virtual bool ContainRecursive(TArray<UAnimCompositeBase*>& CurrentAccumulatedList) PURE_VIRTUAL(UAnimCompositeBase::ContainRecursive, return false; );
 
 	virtual void SetCompositeLength(float InLength) PURE_VIRTUAL(UAnimCompositeBase::SetCompositeLength, );
+
+#if WITH_EDITOR
+	virtual void PopulateWithExistingModel(TScriptInterface<IAnimationDataModel> ExistingDataModel) override;
+#endif // WITH_EDITOR
 };
 
