CIS EdTest Warning: Ensure condition failed: NameMapping->GetName(Curve.Name.UID, BoneName)
#fix change AnimationData plugin to UncookedOnly to fix -Game runs, and address exposed underlying issue of FSmartName UIDs not being copied / updated during UAnimDataModel upgrade-path

diff --git a/Engine/Plugins/Animation/AnimationData/AnimationData.uplugin b/Engine/Plugins/Animation/AnimationData/AnimationData.uplugin
--- a/Engine/Plugins/Animation/AnimationData/AnimationData.uplugin
+++ b/Engine/Plugins/Animation/AnimationData/AnimationData.uplugin
@@ -18,7 +18,7 @@
 	[
 		{
 			"Name": "AnimationData",
-			"Type": "Editor",
+			"Type": "UncookedOnly",
 			"LoadingPhase": "PreDefault"
 		}
 	],
diff --git a/Engine/Source/Developer/AnimationDataController/Private/AnimDataController.cpp b/Engine/Source/Developer/AnimationDataController/Private/AnimDataController.cpp
--- a/Engine/Source/Developer/AnimationDataController/Private/AnimDataController.cpp
+++ b/Engine/Source/Developer/AnimationDataController/Private/AnimDataController.cpp
@@ -2087,6 +2087,13 @@ void UAnimDataController::PopulateWithExistingModel(TScriptInterface<IAnimationD
 	Model->NumberOfKeys = InModel->GetNumberOfFrames() + 1;
 	Model->CurveData = InModel->GetCurveData();
 	Model->AnimatedBoneAttributes = InModel->GetAttributes();
+
+	// Ensure curve name UIDs are validated against the outer skeleton as they are not copied over through FSmartName
+	if (const UAnimSequenceBase* AnimSequenceBase = Cast<UAnimSequenceBase>(InModel.GetObject()->GetOuter()))
+	{
+		FindOrAddCurveNamesOnSkeleton(AnimSequenceBase->GetSkeleton(), ERawCurveTrackTypes::RCT_Float);
+		FindOrAddCurveNamesOnSkeleton(AnimSequenceBase->GetSkeleton(), ERawCurveTrackTypes::RCT_Transform);
+	}
 }
 
 #endif // WITH_EDITOR
