[Backout] - CL20593017, 20594056, 20625634, 20645670, 20649179, 20649223, 20649380, 20658535
Add new custom version to specify a certain window of AnimSequences are not loadable
Original CL Desc
-----------------------------------------------------------------
**New**
- Engine
	- IAnimationDataModel interface
- AnimationData plugin
	- AnimSequencerDataModel, sequencer based implementation of IAnimationDataModel
	- AnimSequencerDataController, controller for above (implementation of IAnimationDataController)
- Added FCompressedRichCurve::PopulateCurve, allowing users to convert back to a FRichCurve for validation/debugging
- Added DefaultFrameRate to AnimationSettings, this replaces the hardcoded 30FPS in code
- Added ::GetKeyIndicesFromTime which takes FFrameRate
- Added AnimSequenceBase::OnAnimModelLoaded, this is required for correct postloading behaviour of data model (specifically AnimSequencerDataModel, as it depends on data from its outer ? anim sequence)
- Added IAnimationDataModel ::Evaluate which now takes responsibility for evaluating raw animation bone, curve and attribute data. And moved all trackbased evaluation code into AnimDataModel.cpp
- Added a.ForceEvalRawData allowing to force evaluation of source data
- Added a.SkipDDC to force compressing animation data locally

**Changed**
- Reparent UAnimDataModel to IAnimationDataModel interface, and rejig behaviour
- AnimSequenceBase now reference AnimDataModel as IAnimationDataModel
- Upgrade path for legacy to IAnimationDataModel and existing UAnimDataModel to IAnimationDataModel through IAnimationDataController::PopulateWithExisting
- IAnimationDataModel data is now frame(number/rate/time) based rather than seconds/keyindices. This enforces frameborder aligned data from now on.
- Moved RichCurve evaluation code from RichCurve.cpp to separate file and consolidated all instances of copy-pasta behaviour
- Updated/removed deprecated EngineTests around AnimSequences
- Changed many instances to use a FFrameTime together with a known FFrameRate versus seconds and sequencelength in float involving frame \/ time calculations (including instances of FAnimKeyHelper)
- Switched anim sequence evaluation to use double in Extraction context for time value and patched up places with static\_cast\<double\>
- FRichCurve::Eval
	- Make sure we always evaluate keys when T >= Key0.Time and T <= KeyN.Time to deal with crazy userweighted tangents
- Fixed WeightedKeyDataAdapter::GetKeyInterpMode/GetKeyTangentWeightMode retrieving invalid values, as it was indexing according to KeyIndex rather than KeyIndex\*2. This made it so that values were misinterpreted.
- Fixed WeightedEvalForTwoKeys for compressed data retrieving GetKeyInterpMode and GetKeyTangentWeightMode using the wrong index value. As they are not indexed using KeyDataHandle type.
- Fixed issue where compression could crash when containing zero-key scale additivebase track (tries to retrieve)
- Replaced instances of NULL with nullptr
- Moved required decompression information into FAnimSequenceDecompressionContext and reimplemented decompression within UE::Anim::Decompression namespace and new file
- Deprecated ::GetRawDataGuid and replaced with GetDataModel()->GenerateGuid()
- Fixed UAnimStreamable evaluation/compression (previously broken with MVC refactor)
- Updated Animation exporting pipeline to be frame-based rather than seconds
- Deprecated RetargetPose, replaced with FRetargetingScope (used within AnimModel::Evaluate)
- Fixed BaseAdditiveAnimation array become invalid/incorrect when removing zero-additve tracks during compression
- Updated EngineTest animation sequence related tests to new APIs (while maintaining deprecated path testing as well for now)
- AnimDataController / Animation Sequence tests now generate a USkeleton for the transient animation sequence (used to perform the test on)

**Removed**
- Unused file/class AnimData/AnimDataNotifyCollector.h

diff --git a/Engine/Plugins/Animation/ControlRig/Source/ControlRig/Private/Sequencer/MovieSceneControlRigParameterSection.cpp b/Engine/Plugins/Animation/ControlRig/Source/ControlRig/Private/Sequencer/MovieSceneControlRigParameterSection.cpp
--- a/Engine/Plugins/Animation/ControlRig/Source/ControlRig/Private/Sequencer/MovieSceneControlRigParameterSection.cpp
+++ b/Engine/Plugins/Animation/ControlRig/Source/ControlRig/Private/Sequencer/MovieSceneControlRigParameterSection.cpp
@@ -2626,11 +2626,11 @@ bool UMovieSceneControlRigParameterSection::LoadAnimSequenceIntoThisSection(UAni
 	}
 	ControlRig->Modify();
 
-	const int32 NumberOfKeys = AnimSequence->GetDataModel()->GetNumberOfKeys();
+	const int32 NumberOfFrames = FrameRate.AsFrameTime(Length).CeilToFrame().Value + 1;
 	FFrameNumber FrameRateInFrameNumber = TickResolution.AsFrameNumber(FrameRate.AsInterval());
 	int32 ExtraProgress = bKeyReduce ? FloatChannels.Num() : 0;
 	
-	FScopedSlowTask Progress(NumberOfKeys + ExtraProgress, LOCTEXT("BakingToControlRig_SlowTask", "Baking To Control Rig..."));	
+	FScopedSlowTask Progress(NumberOfFrames + ExtraProgress, LOCTEXT("BakingToControlRig_SlowTask", "Baking To Control Rig..."));	
 	Progress.MakeDialog(true);
 
 	//Make sure we are reset and run construction event  before evaluating
@@ -2652,7 +2652,7 @@ bool UMovieSceneControlRigParameterSection::LoadAnimSequenceIntoThisSection(UAni
 	SourceCurves.ResetValues();
 	ControlRig->Execute(EControlRigState::Update, TEXT("Setup"));
 	*/
-	const IAnimationDataModel* DataModel = AnimSequence->GetDataModel();
+	const UAnimDataModel* DataModel = AnimSequence->GetDataModel();
 	const FAnimationCurveData& CurveData = DataModel->GetCurveData();
 	const TArray<FBoneAnimationTrack>& BoneAnimationTracks = DataModel->GetBoneAnimationTracks();
 
@@ -2678,10 +2678,10 @@ bool UMovieSceneControlRigParameterSection::LoadAnimSequenceIntoThisSection(UAni
 	ControlRig->RequestConstruction();
 	ControlRig->Evaluate_AnyThread();
 
-	for (int32 Index = 0; Index < NumberOfKeys; ++Index)
+	for (int32 Index = 0; Index < NumberOfFrames; ++Index)
 	{
 		const float SequenceSecond = AnimSequence->GetTimeAtFrame(Index);
-		const FFrameNumber FrameNumber = StartFrame + (FMath::Max(FrameRateInFrameNumber.Value, 1) * Index);
+		FFrameNumber FrameNumber = StartFrame + (FrameRateInFrameNumber * Index);
 
 		ControlRig->GetHierarchy()->ResetPoseToInitial();
 		ControlRig->GetHierarchy()->ResetCurveValues();
diff --git a/Engine/Plugins/Experimental/Animation/ContextualAnimation/Source/ContextualAnimation/Private/ContextualAnimTypes.cpp b/Engine/Plugins/Experimental/Animation/ContextualAnimation/Source/ContextualAnimation/Private/ContextualAnimTypes.cpp
--- a/Engine/Plugins/Experimental/Animation/ContextualAnimation/Source/ContextualAnimation/Private/ContextualAnimTypes.cpp
+++ b/Engine/Plugins/Experimental/Animation/ContextualAnimation/Source/ContextualAnimation/Private/ContextualAnimTypes.cpp
@@ -38,7 +38,7 @@ FTransform FContextualAnimAlignmentTrackContainer::ExtractTransformAtTime(int32
 		const FRawAnimSequenceTrack& Track = Tracks.AnimationTracks[TrackIndex];
 		const int32 TotalFrames = Track.PosKeys.Num();
 		const float TrackLength = (TotalFrames - 1) * SampleInterval;
-		FAnimationUtils::ExtractTransformFromTrack(Track, Time, TotalFrames, TrackLength, EAnimInterpolationType::Linear, AlignmentTransform);
+		FAnimationUtils::ExtractTransformFromTrack(Time, TotalFrames, TrackLength, Track, EAnimInterpolationType::Linear, AlignmentTransform);
 	}
 
 	return AlignmentTransform;
diff --git a/Engine/Plugins/Experimental/Animation/ContextualAnimation/Source/ContextualAnimation/Private/ContextualAnimUtilities.cpp b/Engine/Plugins/Experimental/Animation/ContextualAnimation/Source/ContextualAnimation/Private/ContextualAnimUtilities.cpp
--- a/Engine/Plugins/Experimental/Animation/ContextualAnimation/Source/ContextualAnimation/Private/ContextualAnimUtilities.cpp
+++ b/Engine/Plugins/Experimental/Animation/ContextualAnimation/Source/ContextualAnimation/Private/ContextualAnimUtilities.cpp
@@ -28,7 +28,7 @@ void UContextualAnimUtilities::ExtractLocalSpacePose(const UAnimSequenceBase* An
 	FBlendedCurve Curve;
 	Curve.InitFrom(BoneContainer);
 
-	FAnimExtractContext Context(static_cast<double>(Time), bExtractRootMotion);
+	FAnimExtractContext Context(Time, bExtractRootMotion);
 
 	UE::Anim::FStackAttributeContainer Attributes;
 	FAnimationPoseData AnimationPoseData(OutPose, Curve, Attributes);
diff --git a/Engine/Plugins/Experimental/Animation/MotionWarping/Source/MotionWarping/Private/MotionWarpingComponent.cpp b/Engine/Plugins/Experimental/Animation/MotionWarping/Source/MotionWarping/Private/MotionWarpingComponent.cpp
--- a/Engine/Plugins/Experimental/Animation/MotionWarping/Source/MotionWarping/Private/MotionWarpingComponent.cpp
+++ b/Engine/Plugins/Experimental/Animation/MotionWarping/Source/MotionWarping/Private/MotionWarpingComponent.cpp
@@ -31,7 +31,7 @@ void UMotionWarpingUtilities::ExtractLocalSpacePose(const UAnimSequenceBase* Ani
 	FBlendedCurve Curve;
 	Curve.InitFrom(BoneContainer);
 
-	FAnimExtractContext Context(static_cast<double>(Time), bExtractRootMotion);
+	FAnimExtractContext Context(Time, bExtractRootMotion);
 
 	UE::Anim::FStackAttributeContainer Attributes;
 	FAnimationPoseData AnimationPoseData(OutPose, Curve, Attributes);
diff --git a/Engine/Plugins/Experimental/Animation/MotionWarping/Source/MotionWarping/Private/RootMotionModifier_AdjustmentBlendWarp.cpp b/Engine/Plugins/Experimental/Animation/MotionWarping/Source/MotionWarping/Private/RootMotionModifier_AdjustmentBlendWarp.cpp
--- a/Engine/Plugins/Experimental/Animation/MotionWarping/Source/MotionWarping/Private/RootMotionModifier_AdjustmentBlendWarp.cpp
+++ b/Engine/Plugins/Experimental/Animation/MotionWarping/Source/MotionWarping/Private/RootMotionModifier_AdjustmentBlendWarp.cpp
@@ -37,7 +37,7 @@ void URootMotionModifier_AdjustmentBlendWarp::ExtractBoneTransformAtTime(FTransf
 	const float TrackLength = (EndTime - ActualStartTime);
 	const float TimePercent = (Time - ActualStartTime) / (EndTime - ActualStartTime);
 	const float RemappedTime = TimePercent * TrackLength;
-	FAnimationUtils::ExtractTransformFromTrack(Result.AnimationTracks[TrackIndex], static_cast<double>(RemappedTime), TotalFrames, TrackLength, EAnimInterpolationType::Linear, OutTransform);
+	FAnimationUtils::ExtractTransformFromTrack(RemappedTime, TotalFrames, TrackLength, Result.AnimationTracks[TrackIndex], EAnimInterpolationType::Linear, OutTransform);
 }
 
 void URootMotionModifier_AdjustmentBlendWarp::ExtractBoneTransformAtFrame(FTransform& OutTransform, int32 TrackIndex, int32 Frame) const
diff --git a/Engine/Plugins/Experimental/Animation/PoseSearch/Source/Runtime/Private/PoseSearchDerivedData.cpp b/Engine/Plugins/Experimental/Animation/PoseSearch/Source/Runtime/Private/PoseSearchDerivedData.cpp
--- a/Engine/Plugins/Experimental/Animation/PoseSearch/Source/Runtime/Private/PoseSearchDerivedData.cpp
+++ b/Engine/Plugins/Experimental/Animation/PoseSearch/Source/Runtime/Private/PoseSearchDerivedData.cpp
@@ -279,7 +279,7 @@ namespace UE::PoseSearch
 		if (Sequence)
 		{
 			InOutWriter.Update(MakeMemoryView(Sequence->GetName()));
-			InOutWriter.Update(MakeMemoryView(Sequence->GetDataModel()->GenerateGuid().ToString()));
+			InOutWriter.Update(MakeMemoryView(Sequence->GetRawDataGuid().ToString()));
 		}
 	}
 
diff --git a/Engine/Plugins/Experimental/Animation/PoseSearch/Source/Runtime/Private/PoseSearchFeatureChannels.cpp b/Engine/Plugins/Experimental/Animation/PoseSearch/Source/Runtime/Private/PoseSearchFeatureChannels.cpp
--- a/Engine/Plugins/Experimental/Animation/PoseSearch/Source/Runtime/Private/PoseSearchFeatureChannels.cpp
+++ b/Engine/Plugins/Experimental/Animation/PoseSearch/Source/Runtime/Private/PoseSearchFeatureChannels.cpp
@@ -523,7 +523,7 @@ void UPoseSearchFeatureChannel_Pose::AddPoseFeatures(const  UE::PoseSearch::IAss
 
 			FDeltaTimeRecord DeltaTimeRecord;
 			DeltaTimeRecord.Set(PreviousTime, CurrentTime - PreviousTime);
-			FAnimExtractContext ExtractionCtx(static_cast<double>(CurrentTime), true, DeltaTimeRecord, Samples[Term].Clip->IsLoopable());
+			FAnimExtractContext ExtractionCtx(CurrentTime, true, DeltaTimeRecord, Samples[Term].Clip->IsLoopable());
 
 			Samples[Term].Clip->ExtractPose(ExtractionCtx, AnimPoseData[Term]);
 
diff --git a/Engine/Plugins/Experimental/MLDeformer/MLDeformerFramework/Source/MLDeformerFrameworkEditor/Private/MLDeformerEditorModel.cpp b/Engine/Plugins/Experimental/MLDeformer/MLDeformerFramework/Source/MLDeformerFrameworkEditor/Private/MLDeformerEditorModel.cpp
--- a/Engine/Plugins/Experimental/MLDeformer/MLDeformerFramework/Source/MLDeformerFrameworkEditor/Private/MLDeformerEditorModel.cpp
+++ b/Engine/Plugins/Experimental/MLDeformer/MLDeformerFramework/Source/MLDeformerFrameworkEditor/Private/MLDeformerEditorModel.cpp
@@ -869,7 +869,7 @@ namespace UE::MLDeformer
 			return;
 		}
 
-		const IAnimationDataModel* DataModel = Model->AnimSequence->GetDataModel();
+		const UAnimDataModel* DataModel = Model->AnimSequence->GetDataModel();
 		if (!DataModel)
 		{
 			UE_LOG(LogMLDeformer, Warning, TEXT("Anim sequence has no data model."));
@@ -957,7 +957,7 @@ namespace UE::MLDeformer
 			return;
 		}
 
-		const IAnimationDataModel* DataModel = Model->AnimSequence->GetDataModel();
+		const UAnimDataModel* DataModel = Model->AnimSequence->GetDataModel();
 		if (!DataModel)
 		{
 			UE_LOG(LogMLDeformer, Warning, TEXT("Anim sequence has no data model."));
diff --git a/Engine/Plugins/Experimental/SkeletalReduction/Source/Private/SkeletalMeshReductionPlugin.cpp b/Engine/Plugins/Experimental/SkeletalReduction/Source/Private/SkeletalMeshReductionPlugin.cpp
--- a/Engine/Plugins/Experimental/SkeletalReduction/Source/Private/SkeletalMeshReductionPlugin.cpp
+++ b/Engine/Plugins/Experimental/SkeletalReduction/Source/Private/SkeletalMeshReductionPlugin.cpp
@@ -1202,7 +1202,7 @@ float FQuadricSkeletalMeshReduction::SimplifyMesh( const FSkeletalMeshOptimizati
 
 	// Copy the simplified mesh back into Mesh
 
-	Simplifier.OutputMesh(Mesh.VertexBuffer, Mesh.IndexBuffer, bWeldVertexColorAttributes, nullptr);
+	Simplifier.OutputMesh(Mesh.VertexBuffer, Mesh.IndexBuffer, bWeldVertexColorAttributes, NULL);
 
 	// There might have some unused verts at the end of the vertex buffer that were generated by the possible duplicates
 
@@ -1894,7 +1894,7 @@ void FQuadricSkeletalMeshReduction::ReduceSkeletalMesh(USkeletalMesh& SkeletalMe
 	// Insert a new LOD model entry if needed.
 	if (LODIndex == SkeletalMeshResource.LODModels.Num())
 	{
-		FSkeletalMeshLODModel* ModelPtr = nullptr;
+		FSkeletalMeshLODModel* ModelPtr = NULL;
 		SkeletalMeshResource.LODModels.Add(ModelPtr);
 		bLODModelAdded = true;
 	}
@@ -2119,7 +2119,7 @@ void FQuadricSkeletalMeshReduction::ReduceSkeletalMesh(USkeletalMesh& SkeletalMe
 
 		// Get component space retarget base pose, will be equivalent of ref-pose if not edited
 		TArray<FTransform> ComponentSpaceRefPose;
-		FAnimationRuntime::FillUpComponentSpaceTransforms(SkeletalMesh.GetRefSkeleton(), SkeletalMesh.GetRefSkeleton().GetRefBonePose(), ComponentSpaceRefPose);
+		FAnimationRuntime::FillUpComponentSpaceTransformsRetargetBasePose(&SkeletalMesh, ComponentSpaceRefPose); 
 		
 		// Retrieve bone indices which will be removed
 		if (MeshBoneReductionInterface != nullptr)
@@ -2142,16 +2142,10 @@ void FQuadricSkeletalMeshReduction::ReduceSkeletalMesh(USkeletalMesh& SkeletalMe
 		Pose.SetBoneContainer(&RequiredBones);
 		Pose.ResetToRefPose();
 
-		FBlendedCurve TempCurve;
-		TempCurve.InitFrom(RequiredBones);
-		UE::Anim::FStackAttributeContainer TempAttributes;
-		
-		FAnimationPoseData AnimPose(Pose, TempCurve, TempAttributes);
-
 		// Retrieve animated pose from anim sequence (including retargeting)
 		const USkeleton* Skeleton = SkeletalMesh.GetSkeleton();
 		const FName RetargetSource = Skeleton->GetRetargetSourceForMesh(&SkeletalMesh);
-		UE::Anim::BuildPoseFromModel(BakePoseAnim->GetDataModel(), AnimPose, 0.0, EAnimInterpolationType::Step, RetargetSource, Skeleton->GetRefLocalPoses(RetargetSource));
+		UE::Anim::BuildPoseFromModel(BakePoseAnim->GetDataModel(), Pose, 0.f, EAnimInterpolationType::Step, RetargetSource, Skeleton->GetRefLocalPoses(RetargetSource));
 		
 		// Calculate component space animated pose matrices
 		TArray<FMatrix> ComponentSpaceAnimatedPose;
@@ -2252,7 +2246,7 @@ void FQuadricSkeletalMeshReduction::ReduceSkeletalMesh(USkeletalMesh& SkeletalMe
 		check(ReducedLODInfoPtr);
 		// Do any joint-welding / bone removal.
 
-		if (MeshBoneReductionInterface != nullptr && MeshBoneReductionInterface->GetBoneReductionData(&SkeletalMesh, LODIndex, BonesToRemove))
+		if (MeshBoneReductionInterface != NULL && MeshBoneReductionInterface->GetBoneReductionData(&SkeletalMesh, LODIndex, BonesToRemove))
 		{
 			// fix up chunks to remove the bones that set to be removed
 			for (int32 SectionIndex = 0; SectionIndex < NewModel->Sections.Num(); ++SectionIndex)
@@ -2295,7 +2289,7 @@ void FQuadricSkeletalMeshReduction::ReduceSkeletalMesh(USkeletalMesh& SkeletalMe
 		FSkeletalMeshLODModel::CopyStructure(NewModel, SrcModel);
 
 		// Do any joint-welding / bone removal.
-		if (MeshBoneReductionInterface != nullptr && MeshBoneReductionInterface->GetBoneReductionData(&SkeletalMesh, LODIndex, BonesToRemove))
+		if (MeshBoneReductionInterface != NULL && MeshBoneReductionInterface->GetBoneReductionData(&SkeletalMesh, LODIndex, BonesToRemove))
 		{
 			// fix up chunks to remove the bones that set to be removed
 			for (int32 SectionIndex = 0; SectionIndex < NewModel->Sections.Num(); ++SectionIndex)
diff --git a/Engine/Plugins/Importers/AlembicImporter/Source/AlembicLibrary/Private/AbcImporter.cpp b/Engine/Plugins/Importers/AlembicImporter/Source/AlembicLibrary/Private/AbcImporter.cpp
--- a/Engine/Plugins/Importers/AlembicImporter/Source/AlembicLibrary/Private/AbcImporter.cpp
+++ b/Engine/Plugins/Importers/AlembicImporter/Source/AlembicLibrary/Private/AbcImporter.cpp
@@ -575,7 +575,7 @@ TArray<UObject*> FAbcImporter::ImportAsSkeletalMesh(UObject* InParent, EObjectFl
 	if (SkeletalMesh)
 	{
 		// Touch pre edit change
-		SkeletalMesh->PreEditChange(nullptr);
+		SkeletalMesh->PreEditChange(NULL);
 
 		// Retrieve the imported resource structure and allocate a new LOD model
 		FSkeletalMeshModel* ImportedModel = SkeletalMesh->GetImportedModel();
@@ -651,15 +651,12 @@ TArray<UObject*> FAbcImporter::ImportAsSkeletalMesh(UObject* InParent, EObjectFl
 		IAnimationDataController& Controller = Sequence->GetController();
 
 		Controller.OpenBracket(LOCTEXT("ImportAsSkeletalMesh", "Importing Alembic Animation"));
-		Controller.InitializeModel();
 
-		const FFrameRate FrameRate(AbcFile->GetFramerate(), 1);
-		Controller.SetFrameRate(FrameRate);	
-		const FFrameNumber FrameNumber = FrameRate.AsFrameNumber(AbcFile->GetImportLength());
-		Controller.SetNumberOfFrames(FrameNumber);
+		Controller.SetPlayLength(AbcFile->GetImportLength());
+		Controller.SetFrameRate( FFrameRate(AbcFile->GetFramerate(), 1));
 
-		Sequence->ImportFileFramerate = FrameRate.AsDecimal();
-		Sequence->ImportResampleFramerate = FrameRate.AsInterval();
+		Sequence->ImportFileFramerate = AbcFile->GetFramerate();
+		Sequence->ImportResampleFramerate = 1.0f / (float)AbcFile->GetFramerate();
 
 		{
 #if WITH_EDITOR
@@ -1751,7 +1748,7 @@ bool FAbcImporter::BuildSkeletalMesh( FSkeletalMeshLODModel& LODModel, const FRe
 	}
 
 	// Compute the required bones for this model.
-	USkeletalMesh::CalculateRequiredBones(LODModel, RefSkeleton, nullptr);
+	USkeletalMesh::CalculateRequiredBones(LODModel, RefSkeleton, NULL);
 
 	return true;
 }
diff --git a/Engine/Plugins/Importers/USDImporter/Source/USDUtilities/Private/USDSkeletalDataConversion.cpp b/Engine/Plugins/Importers/USDImporter/Source/USDUtilities/Private/USDSkeletalDataConversion.cpp
--- a/Engine/Plugins/Importers/USDImporter/Source/USDUtilities/Private/USDSkeletalDataConversion.cpp
+++ b/Engine/Plugins/Importers/USDImporter/Source/USDUtilities/Private/USDSkeletalDataConversion.cpp
@@ -429,11 +429,11 @@ namespace SkelDataConversionImpl
 		Skeleton->AddSmartNameAndModify( USkeleton::AnimCurveMappingName, CurveName, NewName );
 
 		const bool bShouldTransact = false;
-		const IAnimationDataModel* DataModel = Sequence->GetDataModel();
+		const UAnimDataModel* DataModel = Sequence->GetDataModel();
 		IAnimationDataController& Controller = Sequence->GetController();
 
-		FAnimationCurveIdentifier CurveId(NewName, ERawCurveTrackTypes::RCT_Float);
-		const FFloatCurve* Curve = DataModel->FindFloatCurve(CurveId);
+		FAnimationCurveIdentifier CurveId( NewName, ERawCurveTrackTypes::RCT_Float );
+		const FFloatCurve* Curve = DataModel->FindFloatCurve( CurveId );
 		if ( !Curve )
 		{
 			// If curve doesn't exist, add one
@@ -1922,7 +1922,6 @@ bool UsdToUnreal::ConvertSkelAnim( const pxr::UsdSkelSkeletonQuery& InUsdSkeleto
 	// it will also create a transaction when importing into UE assets, and the level sequence assets can emit some warnings about it
 	const bool bShouldTransact = false;
 	Controller.OpenBracket( LOCTEXT( "ImportUSDAnimData_Bracket", "Importing USD Animation Data" ), bShouldTransact );
-	Controller.InitializeModel();
 	Controller.ResetModel( bShouldTransact );
 
 	// Bake the animation for each frame.
@@ -2155,11 +2154,8 @@ bool UsdToUnreal::ConvertSkelAnim( const pxr::UsdSkelSkeletonQuery& InUsdSkeleto
 	OutSkeletalAnimationAsset->ImportFileFramerate = LayerTimeCodesPerSecond;
 	OutSkeletalAnimationAsset->ImportResampleFramerate = LayerTimeCodesPerSecond;
 
-
-	const FFrameRate FrameRate(LayerTimeCodesPerSecond, 1);
-	Controller.SetFrameRate(FrameRate, bShouldTransact);
-	const FFrameNumber FrameNumber = FrameRate.AsFrameNumber(LayerSequenceLengthSeconds);
-	Controller.SetNumberOfFrames(FrameNumber, bShouldTransact);
+	Controller.SetPlayLength( LayerSequenceLengthSeconds, bShouldTransact );
+	Controller.SetFrameRate( FFrameRate( LayerTimeCodesPerSecond, 1 ), bShouldTransact );
 	Controller.NotifyPopulated(); // This call is important to get the controller to not use the sampling frequency as framerate
 	Controller.CloseBracket( bShouldTransact );
 
diff --git a/Engine/Plugins/VirtualProduction/Takes/Source/TakeTrackRecorders/Private/MovieScene3DTransformTrackRecorder.cpp b/Engine/Plugins/VirtualProduction/Takes/Source/TakeTrackRecorders/Private/MovieScene3DTransformTrackRecorder.cpp
--- a/Engine/Plugins/VirtualProduction/Takes/Source/TakeTrackRecorders/Private/MovieScene3DTransformTrackRecorder.cpp
+++ b/Engine/Plugins/VirtualProduction/Takes/Source/TakeTrackRecorders/Private/MovieScene3DTransformTrackRecorder.cpp
@@ -12,7 +12,7 @@
 #include "Components/SceneComponent.h"
 #include "Components/SkeletalMeshComponent.h"
 #include "SequenceRecorderUtils.h"
-#include "Animation/AnimData/IAnimationDataModel.h"
+#include "Animation/AnimData/AnimDataModel.h"
 
 
 DEFINE_LOG_CATEGORY(TransformSerialization);
diff --git a/Engine/Plugins/VirtualProduction/Takes/Source/TakeTrackRecorders/Private/MovieSceneAnimationTrackRecorder.cpp b/Engine/Plugins/VirtualProduction/Takes/Source/TakeTrackRecorders/Private/MovieSceneAnimationTrackRecorder.cpp
--- a/Engine/Plugins/VirtualProduction/Takes/Source/TakeTrackRecorders/Private/MovieSceneAnimationTrackRecorder.cpp
+++ b/Engine/Plugins/VirtualProduction/Takes/Source/TakeTrackRecorders/Private/MovieSceneAnimationTrackRecorder.cpp
@@ -390,13 +390,12 @@ bool UMovieSceneAnimationTrackRecorder::LoadRecordedFile(const FString& FileName
 								IAnimationDataController& Controller = AnimSequence->GetController();
 								{
 									IAnimationDataController::FScopedBracket ScopedBracket(Controller, LOCTEXT("LoadRecordedFile_Bracket", "Loading recorded animation file"));
-									Controller.InitializeModel();
+
 									Controller.ResetModel();
 
 									const float FloatDenominator = 1000.0f;
 									const float Numerator = FloatDenominator / Header.IntervalTime;
-									const FFrameRate FrameRate(Numerator, FloatDenominator);
-									Controller.SetFrameRate(FrameRate);
+									Controller.SetFrameRate(FFrameRate(Numerator, FloatDenominator));
 
 									int32 MaxNumberOfKeys = 0;
 									for (int32 TrackIndex = 0; TrackIndex < Header.AnimationTrackNames.Num(); ++TrackIndex)
@@ -423,9 +422,7 @@ bool UMovieSceneAnimationTrackRecorder::LoadRecordedFile(const FString& FileName
 										Controller.SetBoneTrackKeys(Header.AnimationTrackNames[TrackIndex], PosKeys, RotKeys, ScaleKeys);
 									}
 
-									const FFrameNumber FrameNumber = FrameRate.AsFrameNumber((MaxNumberOfKeys > 1) ? (MaxNumberOfKeys - 1) * Header.IntervalTime : FrameRate.AsInterval());
-									Controller.SetNumberOfFrames(FrameNumber);
-
+									Controller.SetPlayLength((MaxNumberOfKeys > 1) ? (MaxNumberOfKeys - 1) * Header.IntervalTime : MINIMUM_ANIMATION_LENGTH);
 									Controller.NotifyPopulated();
 								}
 
@@ -439,7 +436,7 @@ bool UMovieSceneAnimationTrackRecorder::LoadRecordedFile(const FString& FileName
 								FSavePackageArgs SaveArgs;
 								SaveArgs.TopLevelFlags = RF_Standalone;
 								SaveArgs.SaveFlags = SAVE_NoError;
-								UPackage::SavePackage(Package, nullptr, *PackageFileName, SaveArgs);
+								UPackage::SavePackage(Package, NULL, *PackageFileName, SaveArgs);
 
 
 								FFrameNumber SequenceLength = (AnimSequence->GetPlayLength() * TickResolution).FloorToFrame();
diff --git a/Engine/Source/Developer/AnimationDataController/Private/AnimDataController.cpp b/Engine/Source/Developer/AnimationDataController/Private/AnimDataController.cpp
--- a/Engine/Source/Developer/AnimationDataController/Private/AnimDataController.cpp
+++ b/Engine/Source/Developer/AnimationDataController/Private/AnimDataController.cpp
@@ -4,42 +4,89 @@
 #include "AnimDataControllerActions.h"
 #include "MovieSceneTimeHelpers.h"
 
-#include "Animation/AnimData/IAnimationDataModel.h"
+#include "Animation/AnimData/AnimDataModel.h"
 #include "Animation/AnimData/CurveIdentifier.h"
+
+#include "Animation/AnimSequenceBase.h"
 #include "Animation/AnimSequence.h"
 
+
 #include "Algo/Transform.h"
-#include "Animation/AnimationSettings.h"
 #include "UObject/NameTypes.h"
 #include "Animation/AnimCurveTypes.h"
-#include "Animation/AnimSequenceHelpers.h"
 #include "Math/UnrealMathUtility.h"
 
 #define LOCTEXT_NAMESPACE "AnimDataController"
 
 #if WITH_EDITOR
-void UAnimDataController::SetModel(TScriptInterface<IAnimationDataModel> InModel)
+
+namespace UE {
+namespace Anim {
+	bool CanTransactChanges()
+	{
+		return GEngine && GEngine->CanTransact() && !GIsTransacting;
+	}
+
+	struct FScopedCompoundTransaction
+	{
+		FScopedCompoundTransaction(UE::FChangeTransactor& InTransactor, const FText& InDescription) : Transactor(InTransactor), bCreated(false)
+		{
+			if (CanTransactChanges() && !Transactor.IsTransactionPending())
+			{
+				Transactor.OpenTransaction(InDescription);
+				bCreated = true;
+			}
+		}
+
+		~FScopedCompoundTransaction()
+		{
+			if (bCreated)
+			{
+				Transactor.CloseTransaction();
+			}
+		}
+
+		UE::FChangeTransactor& Transactor;
+		bool bCreated;
+	};
+}}
+
+#define CONDITIONAL_TRANSACTION(Text) \
+	TUniquePtr<UE::Anim::FScopedCompoundTransaction> Transaction; \
+	if (UE::Anim::CanTransactChanges() && bShouldTransact) \
+	{ \
+		Transaction = MakeUnique<UE::Anim::FScopedCompoundTransaction>(ChangeTransactor, Text); \
+	}
+
+#define CONDITIONAL_BRACKET(Text) IAnimationDataController::FScopedBracket Transaction(this, Text, UE::Anim::CanTransactChanges() && bShouldTransact);
+
+#define CONDITIONAL_ACTION(ActionClass, ...) \
+	if (UE::Anim::CanTransactChanges() && bShouldTransact) \
+	{ \
+		ChangeTransactor.AddTransactionChange<ActionClass>(__VA_ARGS__); \
+	}
+
+void UAnimDataController::SetModel(UAnimDataModel* InModel)
 {	
 	if (Model != nullptr)
 	{
 		Model->GetModifiedEvent().RemoveAll(this);
 	}
 
-	ModelInterface = InModel;
-	Model = CastChecked<UAnimDataModel>(InModel.GetObject(), ECastCheckedType::NullAllowed);
+	Model = InModel;
 	
-	ChangeTransactor.SetTransactionObject(Model);
+	ChangeTransactor.SetTransactionObject(InModel);
 }
 
 void UAnimDataController::OpenBracket(const FText& InTitle, bool bShouldTransact /*= true*/)
 {
 	ValidateModel();
 
-	if (UE::FChangeTransactor::CanTransactChanges() && !ChangeTransactor.IsTransactionPending())
+	if (UE::Anim::CanTransactChanges() && !ChangeTransactor.IsTransactionPending())
 	{
 		ChangeTransactor.OpenTransaction(InTitle);
 
-		ConditionalAction<UE::Anim::FCloseBracketAction>(bShouldTransact, InTitle.ToString());
+		CONDITIONAL_ACTION(UE::Anim::FCloseBracketAction, InTitle.ToString());
 	}
 
 	if (BracketDepth == 0)
@@ -47,7 +94,7 @@ void UAnimDataController::OpenBracket(const FText& InTitle, bool bShouldTransact
 		FBracketPayload Payload;
 		Payload.Description = InTitle.ToString();
 
-		Model->GetNotifier().Notify(EAnimDataModelNotifyType::BracketOpened, Payload);
+		Model->Notify(EAnimDataModelNotifyType::BracketOpened, Payload);
 	}
 
 	++BracketDepth;
@@ -67,181 +114,95 @@ void UAnimDataController::CloseBracket(bool bShouldTransact /*= true*/)
 
 	if (BracketDepth == 0)
 	{
-		if (UE::FChangeTransactor::CanTransactChanges())
+		if (UE::Anim::CanTransactChanges())
 		{
 			ensure(ChangeTransactor.IsTransactionPending());
 
-			ConditionalAction<UE::Anim::FOpenBracketAction>(bShouldTransact, TEXT("Open Bracket"));
+			CONDITIONAL_ACTION(UE::Anim::FOpenBracketAction, TEXT("Open Bracket"));
 
 			ChangeTransactor.CloseTransaction();
 		}
 		
-		Model->GetNotifier().Notify(EAnimDataModelNotifyType::BracketClosed);
-	}
+		Model->Notify(EAnimDataModelNotifyType::BracketClosed);
+    }
 }
 
-void UAnimDataController::SetNumberOfFrames(FFrameNumber Length, bool bShouldTransact)
+void UAnimDataController::SetPlayLength(float Length, bool bShouldTransact /*= true*/)
 {
 	ValidateModel();
-	const FFrameNumber CurrentNumberOfFrames = Model->GetNumberOfFrames();
-
-	const int32 DeltaFrames = FMath::Abs(Length.Value - CurrentNumberOfFrames.Value);
-	
-	const FFrameNumber T0 = Length > CurrentNumberOfFrames ? CurrentNumberOfFrames : CurrentNumberOfFrames - DeltaFrames;
-	const FFrameNumber T1 = Length > CurrentNumberOfFrames ? Length : CurrentNumberOfFrames;
 
-	ResizeNumberOfFrames(Length, T0, T1, bShouldTransact);
+	// Calculate whether or new play length is shorter or longer than current, set-up T0; T1 accordingly
+	// Assumption is made that time is always added or removed at/from the end
+	// Added: T0 = current length, T1 = new length
+	// Removed: T0 = current length - removed length, T1 = current length
+	const float Delta = FMath::Abs(Length - Model->PlayLength);
+	const float T0 = Length > Model->PlayLength ? Model->PlayLength : Model->PlayLength - Delta;
+	const float T1 = Length > Model->PlayLength ? Length : Model->PlayLength;
+	ResizePlayLength(Length, T0, T1, bShouldTransact);	
 }
 
-void UAnimDataController::ResizeNumberOfFrames(FFrameNumber NewLength, FFrameNumber T0, FFrameNumber T1, bool bShouldTransact)
+void UAnimDataController::Resize(float Length, float T0, float T1, bool bShouldTransact /*= true*/)
 {
 	ValidateModel();
 	
-	const TRange<FFrameNumber> PlayRange(TRange<FFrameNumber>::BoundsType::Inclusive(0), TRange<FFrameNumber>::BoundsType::Exclusive(FMath::Max(1, Model->GetNumberOfKeys())));
-	if (NewLength >= 0)
+	const TRange<float> PlayRange(TRange<float>::BoundsType::Inclusive(0.f), TRange<float>::BoundsType::Inclusive(Model->PlayLength));
+	if (!FMath::IsNearlyZero(Length) && Length > 0.f)
 	{
-		if (NewLength != Model->GetNumberOfFrames())
+		if (Length != Model->PlayLength)
 		{
-			// Ensure that T0 is within the current play range
+			// Ensure that T0 is within the curent play range
 			if (PlayRange.Contains(T0))
 			{
 				// Ensure that the start and end length of either removal or insertion are valid
 				if (T0 < T1)
 				{
-					FTransaction Transaction = ConditionalTransaction(LOCTEXT("ResizePlayLength", "Resizing Play Length"), bShouldTransact);
-
-					const FFrameRate CurrentFrameRate = Model->GetFrameRate();
-					const FFrameNumber CurrentNumberOfFrames = Model->GetNumberOfFrames();
-
-					FSequenceLengthChangedPayload Payload;
-					PRAGMA_DISABLE_DEPRECATION_WARNINGS
-					Payload.T0 = CurrentFrameRate.AsSeconds(T0);
-					Payload.T1 = CurrentFrameRate.AsSeconds(T1);
-					Payload.PreviousLength = CurrentFrameRate.AsSeconds(CurrentNumberOfFrames);
-					PRAGMA_ENABLE_DEPRECATION_WARNINGS
-
-					Payload.PreviousNumberOfFrames = CurrentNumberOfFrames;
-					Payload.Frame0 = T0;
-					Payload.Frame1 = T1;
-
-					ConditionalAction<UE::Anim::FResizePlayLengthInFramesAction>(bShouldTransact, Model, Payload.Frame0, Payload.Frame1);
-
-					Model->NumberOfFrames = NewLength.Value;
-					Model->NumberOfKeys = Model->NumberOfFrames + 1;
-	
-					Model->GetNotifier().Notify<FSequenceLengthChangedPayload>(EAnimDataModelNotifyType::SequenceLengthChanged, Payload);
+					CONDITIONAL_BRACKET(LOCTEXT("ResizeModel", "Resizing Animation Data"));
+					const bool bInserted = Length > Model->PlayLength;
+					ResizePlayLength(Length, T0, T1, bShouldTransact);
+					ResizeCurves(Length, bInserted, T0, T1, bShouldTransact);
+					ResizeAttributes(Length, bInserted, T0, T1, bShouldTransact);
 				}
 				else
 				{
-					ReportErrorf(LOCTEXT("InvalidEndTimeError", "Invalid T1, smaller that T0 value: T0 {0}, T1 {1}"), FText::AsNumber(T0.Value), FText::AsNumber(T0.Value));
+					ReportErrorf(LOCTEXT("InvalidEndTimeError", "Invalid T1, smaller that T0 value: T0 {0}, T1 {1}"), FText::AsNumber(T0), FText::AsNumber(T1));
 				}
 			}
 			else
 			{
-				ReportErrorf(LOCTEXT("InvalidStartTimeError", "Invalid T0, not within existing play range: T0 {0}, Play Length {1}"), FText::AsNumber(T0.Value), FText::AsNumber(Model->GetPlayLength()));
-			}
+				ReportErrorf(LOCTEXT("InvalidStartTimeError", "Invalid T0, not within existing play range: T0 {0}, Play Length {1}"), FText::AsNumber(T0), FText::AsNumber(Model->PlayLength));
+			}			
 		}
-		else if (Model->bPopulated)
+		else
 		{
-			ReportWarningf(LOCTEXT("SamePlayLengthWarning", "New play length is same as existing one: {0} frames"), FText::AsNumber(NewLength.Value));
+			ReportWarningf(LOCTEXT("SamePlayLengthWarning", "New play length is same as existing one: {0} seconds"), FText::AsNumber(Length));
 		}
 	}
 	else
 	{
-		ReportErrorf(LOCTEXT("InvalidPlayLengthError", "Invalid play length value provided: {0} frames"), FText::AsNumber(NewLength.Value));
+		ReportErrorf(LOCTEXT("InvalidPlayLengthError", "Invalid play length value provided: {0} seconds"), FText::AsNumber(Length));
 	}
 }
 
-void UAnimDataController::ResizeInFrames(FFrameNumber NewLength, FFrameNumber T0, FFrameNumber T1, bool bShouldTransact)
-{
-	ValidateModel();
-	
-	const int32 CurrentNumberOFrames = Model->GetNumberOfFrames();
-	
-	const TRange<FFrameNumber> PlayRange(TRange<FFrameNumber>::BoundsType::Inclusive(0), TRange<FFrameNumber>::BoundsType::Exclusive(FMath::Max(1,Model->GetNumberOfKeys())));
-	if (NewLength >= 0)
-	{
-		if (NewLength != Model->GetNumberOfFrames())
-		{
-			// Ensure that T0 is within the current play range
-			if (PlayRange.Contains(T0))
-			{
-				// Ensure that the start and end length of either removal or insertion are valid
-				if (T0 < T1)
-				{
-					FBracket Bracket = ConditionalBracket(LOCTEXT("ResizeModel", "Resizing Animation Data"), bShouldTransact);
-
-					const bool bInserted = NewLength > CurrentNumberOFrames;
-					ResizeNumberOfFrames(NewLength, T0, T1, bShouldTransact);
-					
-					const FFrameRate CurrentFrameRate = Model->GetFrameRate();
-					ResizeCurves(CurrentFrameRate.AsSeconds(NewLength), bInserted, CurrentFrameRate.AsSeconds(T0), CurrentFrameRate.AsSeconds(T1), bShouldTransact);
-					ResizeAttributes(CurrentFrameRate.AsSeconds(NewLength), bInserted, CurrentFrameRate.AsSeconds(T0), CurrentFrameRate.AsSeconds(T1), bShouldTransact);
-				}
-				else
-                {
-                	ReportErrorf(LOCTEXT("InvalidEndTimeError", "Invalid T1, smaller that T0 value: T0 {0}, T1 {1}"), FText::AsNumber(T0.Value), FText::AsNumber(T1.Value));
-                }
-            }
-            else
-            {
-                ReportErrorf(LOCTEXT("InvalidStartTimeError", "Invalid T0, not within existing play range: T0 {0}, Play Length {1}"), FText::AsNumber(T0.Value), FText::AsNumber(CurrentNumberOFrames));
-            }			
-        }
-		else if (Model->bPopulated)
-        {
-            ReportWarningf(LOCTEXT("SameGetPlayLengthWarning", "New play length is same as existing one: {0} frames"), FText::AsNumber(CurrentNumberOFrames));
-        }
-    }
-    else
-    {
-        ReportErrorf(LOCTEXT("InvalidGetPlayLengthError", "Invalid play length value provided: {0} frames"), FText::AsNumber(CurrentNumberOFrames));
-    }
-}
-
-void UAnimDataController::SetPlayLength(float Length, bool bShouldTransact /*= true*/)
-{
-	SetNumberOfFrames(ConvertSecondsToFrameNumber(Length), bShouldTransact);
-}
-
-void UAnimDataController::Resize(float Length, float T0, float T1, bool bShouldTransact /*= true*/)
-{
-	ValidateModel();
-	
-	ResizeInFrames(ConvertSecondsToFrameNumber(Length), ConvertSecondsToFrameNumber(T0), ConvertSecondsToFrameNumber(T1), bShouldTransact);
-}
-
 void UAnimDataController::SetFrameRate(FFrameRate FrameRate, bool bShouldTransact /*= true*/)
 {
 	ValidateModel();
 
 	// Disallow invalid frame-rates, or 0.0 intervals
-	const double FrameRateInterval = FrameRate.AsInterval();
-	if ( FrameRate.IsValid() && !FMath::IsNearlyZero(FrameRateInterval) && FrameRateInterval > 0.0)
+	const float FrameRateInterval = FrameRate.AsInterval();
+	if ( FrameRate.IsValid() && !FMath::IsNearlyZero(FrameRateInterval) && FrameRateInterval > 0.f)
 	{
-		// Need to verify framerate
-		const FFrameRate CurrentFrameRate = Model->GetFrameRate();
-		if (FrameRate.IsMultipleOf(CurrentFrameRate) || FrameRate.IsFactorOf(CurrentFrameRate) || !Model->bPopulated)
-		{
-			FTransaction Transaction = ConditionalTransaction(LOCTEXT("SetFrameRate", "Setting Frame Rate"), bShouldTransact);
-			ConditionalAction<UE::Anim::FSetFrameRateAction>(bShouldTransact, Model);
+		CONDITIONAL_TRANSACTION(LOCTEXT("SetFrameRate", "Setting Frame Rate"));
+
+		CONDITIONAL_ACTION(UE::Anim::FSetFrameRateAction, Model);
 
-			const FFrameNumber CurrentNumberOfFrames = Model->GetNumberOfFrames();
-			const FFrameTime ConvertedLastFrameTime = FFrameRate::TransformTime(CurrentNumberOfFrames, CurrentFrameRate, FrameRate);
-			ensure(FMath::IsNearlyZero(ConvertedLastFrameTime.GetSubFrame()) || !Model->bPopulated);
+		FFrameRateChangedPayload Payload;
+		Payload.PreviousFrameRate = Model->FrameRate;
 			
 			Model->FrameRate = FrameRate;
-			Model->NumberOfFrames = ConvertedLastFrameTime.GetFrame().Value;
+		Model->NumberOfFrames = Model->FrameRate.AsFrameTime(Model->PlayLength).RoundToFrame().Value;
 			Model->NumberOfKeys = Model->NumberOfFrames + 1;
 			
-			FFrameRateChangedPayload Payload;
-			Payload.PreviousFrameRate = CurrentFrameRate;
-			Model->GetNotifier().Notify(EAnimDataModelNotifyType::FrameRateChanged, Payload);
-		}
-		else
-        {
-        	ReportErrorf(LOCTEXT("NonCompatibleFrameRateError", "Incompatible frame rate provided: {0} not a multiple or fact or {1}"), FrameRate.ToPrettyText(), CurrentFrameRate.ToPrettyText());
-        }
+		Model->Notify(EAnimDataModelNotifyType::FrameRateChanged, Payload);
 	}
 	else
 	{
@@ -258,7 +219,7 @@ void UAnimDataController::UpdateCurveNamesFromSkeleton(const USkeleton* Skeleton
 	{
 		if (IsSupportedCurveType(SupportedCurveType))
 		{
-			FBracket Bracket = ConditionalBracket(LOCTEXT("ValidateRawCurves", "Validating Animation Curve Names"), bShouldTransact);
+			CONDITIONAL_BRACKET(LOCTEXT("ValidateRawCurves", "Validating Animation Curve Names"));
 			switch (SupportedCurveType)
 			{
 			case ERawCurveTrackTypes::RCT_Float:
@@ -293,17 +254,12 @@ void UAnimDataController::UpdateCurveNamesFromSkeleton(const USkeleton* Skeleton
 				}
 				break;
 			}
-			case ERawCurveTrackTypes::RCT_Vector:
-			default:
-				const FString CurveTypeAsString = GetCurveTypeValueName(SupportedCurveType);
-				ReportWarningf(LOCTEXT("InvalidCurveTypeWarning", "Invalid curve type provided: {0} ({1})"), FText::FromString(CurveTypeAsString), FText::AsNumber(static_cast<int32>(SupportedCurveType)));
-				
 			}
 		}
 		else
 		{
 			const FString CurveTypeAsString = GetCurveTypeValueName(SupportedCurveType);
-			ReportWarningf(LOCTEXT("InvalidCurveTypeWarning", "Invalid curve type provided: {0} ({1})"), FText::FromString(CurveTypeAsString), FText::AsNumber(static_cast<int32>(SupportedCurveType)));
+			ReportWarningf(LOCTEXT("InvalidCurveTypeWarning", "Invalid curve type provided: {0} ({1})"), FText::FromString(CurveTypeAsString), FText::AsNumber((int32)SupportedCurveType));
 		}
 	}
 	else
@@ -320,7 +276,7 @@ void UAnimDataController::FindOrAddCurveNamesOnSkeleton(USkeleton* Skeleton, ERa
 	{
 		if (IsSupportedCurveType(SupportedCurveType))
 		{
-			FBracket Bracket = ConditionalBracket(LOCTEXT("FindOrAddRawCurveNames", "Updating Skeleton with Animation Curve Names"), bShouldTransact);
+			CONDITIONAL_BRACKET(LOCTEXT("FindOrAddRawCurveNames", "Updating Skeleton with Animation Curve Names"));
 			switch (SupportedCurveType)
 			{
 			case ERawCurveTrackTypes::RCT_Float:
@@ -353,16 +309,12 @@ void UAnimDataController::FindOrAddCurveNamesOnSkeleton(USkeleton* Skeleton, ERa
 				}
 				break;
 			}
-			case ERawCurveTrackTypes::RCT_Vector:
-			default:
-				const FString CurveTypeAsString = GetCurveTypeValueName(SupportedCurveType);
-				ReportWarningf(LOCTEXT("InvalidCurveTypeWarning", "Invalid curve type provided: {0} ({1})"), FText::FromString(CurveTypeAsString), FText::AsNumber(static_cast<int32>(SupportedCurveType)));
 			}
 		}
 		else
 		{
 			const FString CurveTypeAsString = GetCurveTypeValueName(SupportedCurveType);
-			ReportWarningf(LOCTEXT("InvalidCurveTypeWarning", "Invalid curve type provided: {0} ({1})"), FText::FromString(CurveTypeAsString), FText::AsNumber(static_cast<int32>(SupportedCurveType)));
+			ReportWarningf(LOCTEXT("InvalidCurveTypeWarning", "Invalid curve type provided: {0} ({1})"), FText::FromString(CurveTypeAsString), FText::AsNumber((int32)SupportedCurveType));
 		}
 	}
 	else
@@ -373,7 +325,7 @@ void UAnimDataController::FindOrAddCurveNamesOnSkeleton(USkeleton* Skeleton, ERa
 
 bool UAnimDataController::RemoveBoneTracksMissingFromSkeleton(const USkeleton* Skeleton, bool bShouldTransact /*= true*/)
 {
-	if (!ModelInterface->GetAnimationSequence())
+	if (!CheckOuterClass(UAnimSequence::StaticClass()))
 	{
 		return false;
 	}
@@ -405,7 +357,7 @@ bool UAnimDataController::RemoveBoneTracksMissingFromSkeleton(const USkeleton* S
 
 		if (TracksToBeRemoved.Num() || TracksUpdated.Num())
 		{
-			FBracket Bracket = ConditionalBracket(LOCTEXT("RemoveBoneTracksMissingFromSkeleton", "Validating Bone Animation Track Data against Skeleton"), bShouldTransact);
+			CONDITIONAL_BRACKET(LOCTEXT("RemoveBoneTracksMissingFromSkeleton", "Validating Bone Animation Track Data against Skeleton"));
 			for (const FName& TrackName : TracksToBeRemoved)
 			{
 				RemoveBoneTrack(TrackName);
@@ -415,7 +367,7 @@ bool UAnimDataController::RemoveBoneTracksMissingFromSkeleton(const USkeleton* S
 			{
 				FAnimationTrackChangedPayload Payload;
 				Payload.Name = TrackName;
-				Model->GetNotifier().Notify(EAnimDataModelNotifyType::TrackChanged, Payload);
+				Model->Notify(EAnimDataModelNotifyType::TrackChanged, Payload);
 			}
 		}
 
@@ -455,7 +407,7 @@ void UAnimDataController::UpdateAttributesFromSkeleton(const USkeleton* Skeleton
 
 		if (ToRemoveIdentifiers.Num() || ToDuplicateIdentifiers.Num())
 		{
-			FBracket Bracket = ConditionalBracket(LOCTEXT("VerifyAttributeBoneNames", "Remapping Animation Attribute Data"), bShouldTransact);
+			CONDITIONAL_BRACKET(LOCTEXT("VerifyAttributeBoneNames", "Remapping Animation Attribute Data"));
 			for (const FAnimationAttributeIdentifier& Identifier : ToRemoveIdentifiers)
 			{
 				RemoveAttribute(Identifier);
@@ -463,8 +415,8 @@ void UAnimDataController::UpdateAttributesFromSkeleton(const USkeleton* Skeleton
 			
 			for (const TPair<FAnimationAttributeIdentifier, int32>& Pair : ToDuplicateIdentifiers)
 			{
-				const FAnimationAttributeIdentifier& DuplicateIdentifier = Pair.Key;
-				FAnimationAttributeIdentifier NewIdentifier(DuplicateIdentifier.GetName(), Pair.Value, DuplicateIdentifier.GetBoneName(), DuplicateIdentifier.GetType());
+				FAnimationAttributeIdentifier NewIdentifier = Pair.Key;
+				NewIdentifier.BoneIndex = Pair.Value;
 
 				DuplicateAttribute(Pair.Key, NewIdentifier);
 				RemoveAttribute(Pair.Key);
@@ -481,17 +433,17 @@ void UAnimDataController::ResetModel(bool bShouldTransact /*= true*/)
 {
 	ValidateModel();
 
-	FBracket Bracket = ConditionalBracket(LOCTEXT("ResetModel", "Clearing Animation Data"), bShouldTransact);
+	CONDITIONAL_BRACKET(LOCTEXT("ResetModel", "Clearing Animation Data"));
 
 	RemoveAllBoneTracks(bShouldTransact);
 
 	RemoveAllCurvesOfType(ERawCurveTrackTypes::RCT_Float, bShouldTransact);
 	RemoveAllCurvesOfType(ERawCurveTrackTypes::RCT_Transform, bShouldTransact);
 
-	SetFrameRate(UAnimationSettings::Get()->GetDefaultFrameRate(), bShouldTransact);
-	SetNumberOfFrames(1, bShouldTransact);
+	SetPlayLength(MINIMUM_ANIMATION_LENGTH, bShouldTransact);
+	SetFrameRate(FFrameRate(30,1), bShouldTransact);
 
-	Model->GetNotifier().Notify(EAnimDataModelNotifyType::Reset);
+	Model->Notify(EAnimDataModelNotifyType::Reset);
 }
 
 bool UAnimDataController::AddCurve(const FAnimationCurveIdentifier& CurveId, int32 CurveFlags /*= EAnimAssetCurveFlags::AACF_Editable*/, bool bShouldTransact /*= true*/)
@@ -503,7 +455,7 @@ bool UAnimDataController::AddCurve(const FAnimationCurveIdentifier& CurveId, int
 		{
 			if (!Model->FindCurve(CurveId))
 			{
-				FTransaction Transaction = ConditionalTransaction(LOCTEXT("AddRawCurve", "Adding Animation Curve"), bShouldTransact);
+				CONDITIONAL_TRANSACTION(LOCTEXT("AddRawCurve", "Adding Animation Curve"));
 
 				FCurveAddedPayload Payload;
 				Payload.Identifier = CurveId;
@@ -521,27 +473,23 @@ bool UAnimDataController::AddCurve(const FAnimationCurveIdentifier& CurveId, int
 				case ERawCurveTrackTypes::RCT_Float:
 					AddNewCurve(Model->CurveData.FloatCurves);
 					break;
-				case ERawCurveTrackTypes::RCT_Vector:
-				default:
-					const FString CurveTypeAsString = GetCurveTypeValueName(CurveId.CurveType);
-					ReportWarningf(LOCTEXT("InvalidCurveTypeWarning", "Invalid curve type provided: {0} ({1})"), FText::FromString(CurveTypeAsString), FText::AsNumber(static_cast<int32>(CurveId.CurveType)));
 				}
 
-				ConditionalAction<UE::Anim::FRemoveCurveAction>(bShouldTransact, CurveId);
-				Model->GetNotifier().Notify(EAnimDataModelNotifyType::CurveAdded, Payload);
+				CONDITIONAL_ACTION(UE::Anim::FRemoveCurveAction, CurveId);
+				Model->Notify(EAnimDataModelNotifyType::CurveAdded, Payload);
 
 				return true;
 			}
 			else
 			{
 				const FString CurveTypeAsString = GetCurveTypeValueName(CurveId.CurveType);
-				ReportWarningf(LOCTEXT("ExistingCurveNameWarning", "Curve with name {0} and type {1} ({2}) already exists"), FText::FromName(CurveId.InternalName.DisplayName), FText::FromString(CurveTypeAsString), FText::AsNumber(static_cast<int32>(CurveId.CurveType)));
+				ReportWarningf(LOCTEXT("ExistingCurveNameWarning", "Curve with name {0} and type {1} ({2}) already exists"), FText::FromName(CurveId.InternalName.DisplayName), FText::FromString(CurveTypeAsString), FText::AsNumber((int32)CurveId.CurveType));
 			}			
 		}
 		else 
 		{
 			const FString CurveTypeAsString = GetCurveTypeValueName(CurveId.CurveType);
-			ReportWarningf(LOCTEXT("InvalidCurveTypeWarning", "Invalid curve type provided: {0} ({1})"), FText::FromString(CurveTypeAsString), FText::AsNumber(static_cast<int32>(CurveId.CurveType)));
+			ReportWarningf(LOCTEXT("InvalidCurveTypeWarning", "Invalid curve type provided: {0} ({1})"), FText::FromString(CurveTypeAsString), FText::AsNumber((int32)CurveId.CurveType));
 		}		
 	}
 	else
@@ -569,7 +517,7 @@ bool UAnimDataController::DuplicateCurve(const FAnimationCurveIdentifier& CopyCu
 				{
 					if (!Model->FindCurve(NewCurveId))
 					{
-						FTransaction Transaction = ConditionalTransaction(LOCTEXT("CopyRawCurve", "Duplicating Animation Curve"), bShouldTransact);
+						CONDITIONAL_TRANSACTION(LOCTEXT("CopyRawCurve", "Duplicating Animation Curve"));
 
 						auto DuplicateCurve = [NewCurveName = NewCurveId.InternalName](auto& CurveDataArray, const auto& SourceCurve)
 						{
@@ -585,37 +533,33 @@ bool UAnimDataController::DuplicateCurve(const FAnimationCurveIdentifier& CopyCu
 						case ERawCurveTrackTypes::RCT_Float:
 							DuplicateCurve(Model->CurveData.FloatCurves, Model->GetFloatCurve(CopyCurveId));
 							break;
-						case ERawCurveTrackTypes::RCT_Vector:
-						default:
-							const FString CurveTypeAsString = GetCurveTypeValueName(SupportedCurveType);
-							ReportWarningf(LOCTEXT("InvalidCurveTypeWarning", "Invalid curve type provided: {0} ({1})"), FText::FromString(CurveTypeAsString), FText::AsNumber(static_cast<int32>(SupportedCurveType)));
 						}
 
 						FCurveAddedPayload Payload;
 						Payload.Identifier = NewCurveId;
-						Model->GetNotifier().Notify(EAnimDataModelNotifyType::CurveAdded, Payload);
+						Model->Notify(EAnimDataModelNotifyType::CurveAdded, Payload);
 
-						ConditionalAction<UE::Anim::FRemoveCurveAction>(bShouldTransact, NewCurveId);
+						CONDITIONAL_ACTION(UE::Anim::FRemoveCurveAction, NewCurveId);
 
 						return true;
 					}
 					else
 					{
 						const FString CurveTypeAsString = GetCurveTypeValueName(NewCurveId.CurveType);
-						ReportWarningf(LOCTEXT("ExistingCurveNameWarning", "Curve with name {0} and type {1} ({2}) already exists"), FText::FromName(NewCurveId.InternalName.DisplayName), FText::FromString(CurveTypeAsString), FText::AsNumber(static_cast<int32>(NewCurveId.CurveType)));
+						ReportWarningf(LOCTEXT("ExistingCurveNameWarning", "Curve with name {0} and type {1} ({2}) already exists"), FText::FromName(NewCurveId.InternalName.DisplayName), FText::FromString(CurveTypeAsString), FText::AsNumber((int32)NewCurveId.CurveType));
 					}
 				}
 				else
 				{
 					const FString CurveTypeAsString = GetCurveTypeValueName(CopyCurveId.CurveType);
-					ReportWarningf(LOCTEXT("CurveNameToDuplicateNotFoundWarning", "Could not find curve with name {0} and type {1} ({2}) for duplication"), FText::FromName(NewCurveId.InternalName.DisplayName), FText::FromString(CurveTypeAsString), FText::AsNumber(static_cast<int32>(NewCurveId.CurveType)));
+					ReportWarningf(LOCTEXT("CurveNameToDuplicateNotFoundWarning", "Could not find curve with name {0} and type {1} ({2}) for duplication"), FText::FromName(NewCurveId.InternalName.DisplayName), FText::FromString(CurveTypeAsString), FText::AsNumber((int32)NewCurveId.CurveType));
 				}
 			}
 		}
 		else
 		{
 			const FString CurveTypeAsString = GetCurveTypeValueName(SupportedCurveType);
-			ReportWarningf(LOCTEXT("InvalidCurveTypeWarning", "Invalid curve type provided: {0} ({1})"), FText::FromString(CurveTypeAsString), FText::AsNumber(static_cast<int32>(SupportedCurveType)));
+			ReportWarningf(LOCTEXT("InvalidCurveTypeWarning", "Invalid curve type provided: {0} ({1})"), FText::FromString(CurveTypeAsString), FText::AsNumber((int32)SupportedCurveType));
 		}
 	}
 
@@ -632,35 +576,32 @@ bool UAnimDataController::RemoveCurve(const FAnimationCurveIdentifier& CurveId,
 	{
 		if (IsSupportedCurveType(CurveId.CurveType))
 		{
-			if (Model->FindCurve(CurveId) != nullptr)
+			const FAnimCurveBase* Curve = Model->FindCurve(CurveId);
+			if (Curve)
 			{
-				FTransaction Transaction = ConditionalTransaction(LOCTEXT("RemoveCurve", "Removing Animation Curve"), bShouldTransact);
+				CONDITIONAL_TRANSACTION(LOCTEXT("RemoveCurve", "Removing Animation Curve"));
 
 				switch (SupportedCurveType)
 				{
 					case ERawCurveTrackTypes::RCT_Transform:
 					{
 						const FTransformCurve& TransformCurve = Model->GetTransformCurve(CurveId);
-						ConditionalAction<UE::Anim::FAddTransformCurveAction>(bShouldTransact, CurveId, TransformCurve.GetCurveTypeFlags(), TransformCurve);
+						CONDITIONAL_ACTION(UE::Anim::FAddTransformCurveAction, CurveId, TransformCurve.GetCurveTypeFlags(), TransformCurve);
 						Model->CurveData.TransformCurves.RemoveAll([Name = TransformCurve.Name](const FTransformCurve& ToRemoveCurve) { return ToRemoveCurve.Name == Name; });
 						break;
 					}
 					case ERawCurveTrackTypes::RCT_Float:
 					{
 						const FFloatCurve& FloatCurve = Model->GetFloatCurve(CurveId);
-						ConditionalAction<UE::Anim::FAddFloatCurveAction>(bShouldTransact, CurveId, FloatCurve.GetCurveTypeFlags(), FloatCurve.FloatCurve.GetConstRefOfKeys(), FloatCurve.Color);
+						CONDITIONAL_ACTION(UE::Anim::FAddFloatCurveAction, CurveId, FloatCurve.GetCurveTypeFlags(), FloatCurve.FloatCurve.GetConstRefOfKeys(), FloatCurve.Color);
 						Model->CurveData.FloatCurves.RemoveAll([Name = FloatCurve.Name](const FFloatCurve& ToRemoveCurve) { return ToRemoveCurve.Name == Name; });
 						break;
 					}
-					case ERawCurveTrackTypes::RCT_Vector:
-					default:
-						const FString CurveTypeAsString = GetCurveTypeValueName(SupportedCurveType);
-						ReportWarningf(LOCTEXT("InvalidCurveTypeWarning", "Invalid curve type provided: {0} ({1})"), FText::FromString(CurveTypeAsString), FText::AsNumber(static_cast<int32>(SupportedCurveType)));
 				}
 
 				FCurveRemovedPayload Payload;
 				Payload.Identifier = CurveId;
-				Model->GetNotifier().Notify(EAnimDataModelNotifyType::CurveRemoved, Payload);
+				Model->Notify(EAnimDataModelNotifyType::CurveRemoved, Payload);
 
 				return true;
 			}
@@ -673,7 +614,7 @@ bool UAnimDataController::RemoveCurve(const FAnimationCurveIdentifier& CurveId,
 		else
 		{
 			const FString CurveTypeAsString = GetCurveTypeValueName(CurveId.CurveType);
-			ReportWarningf(LOCTEXT("InvalidCurveTypeWarning", "Invalid curve type provided: {0} ({1})"), FText::FromString(CurveTypeAsString), FText::AsNumber(static_cast<int32>(CurveId.CurveType)));
+			ReportWarningf(LOCTEXT("InvalidCurveTypeWarning", "Invalid curve type provided: {0} ({1})"), FText::FromString(CurveTypeAsString), FText::AsNumber((int32)CurveId.CurveType));
 		}
 	}
 
@@ -684,7 +625,7 @@ void UAnimDataController::RemoveAllCurvesOfType(ERawCurveTrackTypes SupportedCur
 {
 	ValidateModel();
 
-	FBracket Bracket = ConditionalBracket(LOCTEXT("DeleteAllRawCurve", "Deleting All Animation Curve"), bShouldTransact);
+	CONDITIONAL_BRACKET(LOCTEXT("DeleteAllRawCurve", "Deleting All Animation Curve"));
 	switch (SupportedCurveType)
 	{
 	case ERawCurveTrackTypes::RCT_Transform:
@@ -708,7 +649,7 @@ void UAnimDataController::RemoveAllCurvesOfType(ERawCurveTrackTypes SupportedCur
 	case ERawCurveTrackTypes::RCT_Vector:
 	default:
 		const FString CurveTypeAsString = GetCurveTypeValueName(SupportedCurveType);
-		ReportWarningf(LOCTEXT("InvalidCurveTypeWarning", "Invalid curve type provided: {0} ({1})"), FText::FromString(CurveTypeAsString), FText::AsNumber(static_cast<int32>(SupportedCurveType)));
+		ReportWarningf(LOCTEXT("InvalidCurveTypeWarning", "Invalid curve type provided: {0} ({1})"), FText::FromString(CurveTypeAsString), FText::AsNumber((int32)SupportedCurveType));
 	}
 
 }
@@ -717,7 +658,7 @@ bool UAnimDataController::SetCurveFlag(const FAnimationCurveIdentifier& CurveId,
 {
 	ValidateModel();
 
-	const ERawCurveTrackTypes SupportedCurveType = CurveId.CurveType;
+	ERawCurveTrackTypes SupportedCurveType = CurveId.CurveType;
 
 	FAnimCurveBase* Curve = nullptr;
 
@@ -729,19 +670,14 @@ bool UAnimDataController::SetCurveFlag(const FAnimationCurveIdentifier& CurveId,
 	{
 		Curve = Model->FindMutableTransformCurveById(CurveId);
 	}
-	else
-	{
-		const FString CurveTypeAsString = GetCurveTypeValueName(SupportedCurveType);
-		ReportWarningf(LOCTEXT("InvalidCurveTypeWarning", "Invalid curve type provided: {0} ({1})"), FText::FromString(CurveTypeAsString), FText::AsNumber(static_cast<int32>(SupportedCurveType)));
-	}
 	
 	if (Curve)
 	{
-		FTransaction Transaction = ConditionalTransaction(LOCTEXT("SetCurveFlag", "Setting Raw Curve Flag"), bShouldTransact);
+		CONDITIONAL_TRANSACTION(LOCTEXT("SetCurveFlag", "Setting Raw Curve Flag"));
 
 		const int32 CurrentFlags = Curve->GetCurveTypeFlags();
 
-		ConditionalAction<UE::Anim::FSetCurveFlagsAction>(bShouldTransact, CurveId, CurrentFlags, SupportedCurveType);
+		CONDITIONAL_ACTION(UE::Anim::FSetCurveFlagsAction, CurveId, CurrentFlags, SupportedCurveType);
 
 		FCurveFlagsChangedPayload Payload;
 		Payload.Identifier = CurveId;
@@ -749,7 +685,7 @@ bool UAnimDataController::SetCurveFlag(const FAnimationCurveIdentifier& CurveId,
 
 		Curve->SetCurveTypeFlag(Flag, bState);
 
-		Model->GetNotifier().Notify(EAnimDataModelNotifyType::CurveFlagsChanged, Payload);
+		Model->Notify(EAnimDataModelNotifyType::CurveFlagsChanged, Payload);
 
 		return true;
 	}
@@ -768,7 +704,8 @@ bool UAnimDataController::SetCurveFlags(const FAnimationCurveIdentifier& CurveId
 
 	FAnimCurveBase* Curve = nullptr;
 
-	const ERawCurveTrackTypes SupportedCurveType = CurveId.CurveType;
+	ERawCurveTrackTypes SupportedCurveType = CurveId.CurveType;
+
 	if (SupportedCurveType == ERawCurveTrackTypes::RCT_Float)
 	{
 		Curve = Model->FindMutableFloatCurveById(CurveId);
@@ -777,19 +714,14 @@ bool UAnimDataController::SetCurveFlags(const FAnimationCurveIdentifier& CurveId
 	{
 		Curve = Model->FindMutableTransformCurveById(CurveId);
 	}
-	else
-	{
-		const FString CurveTypeAsString = GetCurveTypeValueName(SupportedCurveType);
-		ReportWarningf(LOCTEXT("InvalidCurveTypeWarning", "Invalid curve type provided: {0} ({1})"), FText::FromString(CurveTypeAsString), FText::AsNumber(static_cast<int32>(SupportedCurveType)));
-	}
 
 	if (Curve)
 	{
-		FTransaction Transaction = ConditionalTransaction(LOCTEXT("SetRawCurveFlag", "Setting Raw Curve Flags"), bShouldTransact);
+		CONDITIONAL_TRANSACTION(LOCTEXT("SetRawCurveFlag", "Setting Raw Curve Flags"));
 
 		const int32 CurrentFlags = Curve->GetCurveTypeFlags();
 
-		ConditionalAction<UE::Anim::FSetCurveFlagsAction>(bShouldTransact, CurveId, CurrentFlags, SupportedCurveType);
+		CONDITIONAL_ACTION(UE::Anim::FSetCurveFlagsAction, CurveId, CurrentFlags, SupportedCurveType);
 
 		FCurveFlagsChangedPayload Payload;
 		Payload.Identifier = CurveId;
@@ -797,7 +729,7 @@ bool UAnimDataController::SetCurveFlags(const FAnimationCurveIdentifier& CurveId
 
 		Curve->SetCurveTypeFlags(Flags);
 
-		Model->GetNotifier().Notify(EAnimDataModelNotifyType::CurveFlagsChanged, Payload);
+		Model->Notify(EAnimDataModelNotifyType::CurveFlagsChanged, Payload);
 
 		return true;
 	}
@@ -816,9 +748,11 @@ bool UAnimDataController::SetTransformCurveKeys(const FAnimationCurveIdentifier&
 
 	if (TransformValues.Num() == TimeKeys.Num())
 	{
-		if (Model->FindMutableTransformCurveById(CurveId) != nullptr)
+		FTransformCurve* Curve = Model->FindMutableTransformCurveById(CurveId);
+
+		if (Curve)
 		{
-			FBracket Bracket = ConditionalBracket(LOCTEXT("SetTransformCurveKeys_Bracket", "Setting Transform Curve Keys"), bShouldTransact);
+			CONDITIONAL_BRACKET(LOCTEXT("SetTransformCurveKeys_Bracket", "Setting Transform Curve Keys"));
 			
 			struct FKeys
 			{
@@ -864,11 +798,11 @@ bool UAnimDataController::SetTransformCurveKeys(const FAnimationCurveIdentifier&
 			
 			for (int32 SubCurveIndex = 0; SubCurveIndex < 3; ++SubCurveIndex)
 			{
-				const ETransformCurveChannel Channel = static_cast<ETransformCurveChannel>(SubCurveIndex);
-				const FKeys* CurveKeys = SubCurveKeys[SubCurveIndex];
+				const ETransformCurveChannel Channel = (ETransformCurveChannel)SubCurveIndex;
+				FKeys* CurveKeys = SubCurveKeys[SubCurveIndex];
 				for (int32 ChannelIndex = 0; ChannelIndex < 3; ++ChannelIndex)
 				{
-					const EVectorCurveChannel Axis = static_cast<EVectorCurveChannel>(ChannelIndex);
+					const EVectorCurveChannel Axis = (EVectorCurveChannel)ChannelIndex;
 					FAnimationCurveIdentifier TargetCurveIdentifier = CurveId;
 					UAnimationCurveIdentifierExtensions::GetTransformChildCurveIdentifier(TargetCurveIdentifier, Channel, Axis);
 					SetCurveKeys(TargetCurveIdentifier, CurveKeys->ChannelKeys[ChannelIndex], bShouldTransact);
@@ -896,9 +830,11 @@ bool UAnimDataController::SetTransformCurveKey(const FAnimationCurveIdentifier&
 {
 	ValidateModel();
 
-	if (Model->FindMutableTransformCurveById(CurveId) != nullptr)
+	FTransformCurve* Curve = Model->FindMutableTransformCurveById(CurveId);
+
+	if (Curve)
 	{
-		FBracket Bracket = ConditionalBracket(LOCTEXT("AddTransformCurveKey_Bracket", "Setting Transform Curve Key"), bShouldTransact);
+		CONDITIONAL_BRACKET(LOCTEXT("AddTransformCurveKey_Bracket", "Setting Transform Curve Key"));
 		struct FKeys
 		{
 			FRichCurveKey ChannelKeys[3];
@@ -925,11 +861,11 @@ bool UAnimDataController::SetTransformCurveKey(const FAnimationCurveIdentifier&
 		
 		for (int32 SubCurveIndex = 0; SubCurveIndex < 3; ++SubCurveIndex)
 		{
-			const ETransformCurveChannel Channel = static_cast<ETransformCurveChannel>(SubCurveIndex);
+			const ETransformCurveChannel Channel = (ETransformCurveChannel)SubCurveIndex;
 			const FKeys& VectorCurveKeys = VectorKeys[SubCurveIndex];
 			for (int32 ChannelIndex = 0; ChannelIndex < 3; ++ChannelIndex)
 			{
-				const EVectorCurveChannel Axis = static_cast<EVectorCurveChannel>(ChannelIndex);
+				const EVectorCurveChannel Axis = (EVectorCurveChannel)ChannelIndex;
 				FAnimationCurveIdentifier TargetCurveIdentifier = CurveId;
 				UAnimationCurveIdentifierExtensions::GetTransformChildCurveIdentifier(TargetCurveIdentifier, Channel, Axis);
 				SetCurveKey(TargetCurveIdentifier, VectorCurveKeys.ChannelKeys[ChannelIndex], bShouldTransact);
@@ -951,27 +887,30 @@ bool UAnimDataController::RemoveTransformCurveKey(const FAnimationCurveIdentifie
 {
 	ValidateModel();
 
-	if (Model->FindMutableTransformCurveById(CurveId))
+	FTransformCurve* TransformCurve = Model->FindMutableTransformCurveById(CurveId);
+	if (TransformCurve)
 	{
 		const FString BaseCurveName = CurveId.InternalName.DisplayName.ToString();
 		const TArray<FString> SubCurveNames = { TEXT( "Translation"), TEXT( "Rotation"), TEXT( "Scale") };
 		const TArray<FString> ChannelCurveNames = { TEXT("X"), TEXT("Y"), TEXT("Z") };
 
-		FBracket Bracket = ConditionalBracket(LOCTEXT("RemoveTransformCurveKey_Bracket", "Deleting Animation Transform Curve Key"), bShouldTransact);
+		CONDITIONAL_BRACKET(LOCTEXT("RemoveTransformCurveKey_Bracket", "Deleting Animation Transform Curve Key"));
 		
 		for (int32 SubCurveIndex = 0; SubCurveIndex < 3; ++SubCurveIndex)
 		{
-			const ETransformCurveChannel Channel = static_cast<ETransformCurveChannel>(SubCurveIndex);
+			const ETransformCurveChannel Channel = (ETransformCurveChannel)SubCurveIndex;
 			for (int32 ChannelIndex = 0; ChannelIndex < 3; ++ChannelIndex)
 			{
-				const EVectorCurveChannel Axis = static_cast<EVectorCurveChannel>(ChannelIndex);
+				const EVectorCurveChannel Axis = (EVectorCurveChannel)ChannelIndex;
 				FAnimationCurveIdentifier TargetCurveIdentifier = CurveId;
 				UAnimationCurveIdentifierExtensions::GetTransformChildCurveIdentifier(TargetCurveIdentifier, Channel, Axis);
 				RemoveCurveKey(TargetCurveIdentifier, Time, bShouldTransact);
 			}
 		}
 
+
 		return true;
+
 	}
 	else
 	{
@@ -991,9 +930,10 @@ bool UAnimDataController::RenameCurve(const FAnimationCurveIdentifier& CurveToRe
 		{
 			if (CurveToRenameId.CurveType == NewCurveId.CurveType)
 			{
-				if (FAnimCurveBase* Curve = Model->FindMutableCurveById(CurveToRenameId))
+				FAnimCurveBase* Curve = Model->FindMutableCurveById(CurveToRenameId);
+				if (Curve)
 				{
-					FTransaction Transaction = ConditionalTransaction(LOCTEXT("RenameCurve", "Renaming Curve"), bShouldTransact);
+					CONDITIONAL_TRANSACTION(LOCTEXT("RenameCurve", "Renaming Curve"));
 
 					FCurveRenamedPayload Payload;
 					Payload.Identifier = FAnimationCurveIdentifier(Curve->Name, CurveToRenameId.CurveType);
@@ -1001,9 +941,9 @@ bool UAnimDataController::RenameCurve(const FAnimationCurveIdentifier& CurveToRe
 					Curve->Name = NewCurveId.InternalName;
 					Payload.NewIdentifier = NewCurveId;
 
-					ConditionalAction<UE::Anim::FRenameCurveAction>(bShouldTransact, NewCurveId, CurveToRenameId);
+					CONDITIONAL_ACTION(UE::Anim::FRenameCurveAction, NewCurveId, CurveToRenameId);
 
-					Model->GetNotifier().Notify(EAnimDataModelNotifyType::CurveRenamed, Payload);
+					Model->Notify(EAnimDataModelNotifyType::CurveRenamed, Payload);
 
 					return true;
 				}
@@ -1043,17 +983,18 @@ bool UAnimDataController::SetCurveColor(const FAnimationCurveIdentifier& CurveId
 	{
 		if (CurveId.CurveType == ERawCurveTrackTypes::RCT_Float)
 		{
-			if (FFloatCurve* Curve = Model->FindMutableFloatCurveById(CurveId))
+			FFloatCurve* Curve = Model->FindMutableFloatCurveById(CurveId);
+			if (Curve)
 			{
-				FTransaction Transaction = ConditionalTransaction(LOCTEXT("ChangingCurveColor", "Changing Curve Color"), bShouldTransact);
+				CONDITIONAL_TRANSACTION(LOCTEXT("ChangingCurveColor", "Changing Curve Color"));
 
-				ConditionalAction<UE::Anim::FSetCurveColorAction>(bShouldTransact, CurveId, Curve->Color);
+				CONDITIONAL_ACTION(UE::Anim::FSetCurveColorAction, CurveId, Curve->Color);
 
 				Curve->Color = Color;
 
 				FCurveChangedPayload Payload;
 				Payload.Identifier = CurveId;
-				Model->GetNotifier().Notify(EAnimDataModelNotifyType::CurveColorChanged, Payload);
+				Model->Notify(EAnimDataModelNotifyType::CurveColorChanged, Payload);
 
 				return true;				
 			}
@@ -1080,12 +1021,13 @@ bool UAnimDataController::ScaleCurve(const FAnimationCurveIdentifier& CurveId, f
 {
 	ValidateModel();
 
-	const ERawCurveTrackTypes SupportedCurveType = CurveId.CurveType;
+	ERawCurveTrackTypes SupportedCurveType = CurveId.CurveType;
 	if (SupportedCurveType == ERawCurveTrackTypes::RCT_Float)
 	{
-		if (FFloatCurve* Curve = Model->FindMutableFloatCurveById(CurveId))
+		FFloatCurve* Curve = Model->FindMutableFloatCurveById(CurveId);
+		if (Curve)
 		{
-			FTransaction Transaction = ConditionalTransaction(LOCTEXT("ScalingCurve", "Scaling Curve"), bShouldTransact);
+			CONDITIONAL_TRANSACTION(LOCTEXT("ScalingCurve", "Scaling Curve"));
 
 			Curve->FloatCurve.ScaleCurve(Origin, Factor);
 
@@ -1094,9 +1036,9 @@ bool UAnimDataController::ScaleCurve(const FAnimationCurveIdentifier& CurveId, f
 			Payload.Factor = Factor;
 			Payload.Origin = Origin;
 			
-			ConditionalAction<UE::Anim::FScaleCurveAction>(bShouldTransact, CurveId, Origin, 1.0f / Factor, SupportedCurveType);
+			CONDITIONAL_ACTION(UE::Anim::FScaleCurveAction, CurveId, Origin, 1.0f / Factor, SupportedCurveType);
 
-			Model->GetNotifier().Notify(EAnimDataModelNotifyType::CurveScaled, Payload);
+			Model->Notify(EAnimDataModelNotifyType::CurveScaled, Payload);
 
 			return true;
 		}
@@ -1117,7 +1059,9 @@ bool UAnimDataController::SetCurveKey(const FAnimationCurveIdentifier& CurveId,
 {
 	ValidateModel();
 
-	if (FRichCurve* RichCurve = Model->GetMutableRichCurve(CurveId))
+	FRichCurve* RichCurve = Model->GetMutableRichCurve(CurveId);
+	ERawCurveTrackTypes SupportedCurveType = CurveId.CurveType;
+	if (RichCurve)
 	{
 		FCurveChangedPayload Payload;
 		Payload.Identifier = CurveId;
@@ -1126,25 +1070,25 @@ bool UAnimDataController::SetCurveKey(const FAnimationCurveIdentifier& CurveId,
 		const FKeyHandle Handle = RichCurve->FindKey(Key.Time, 0.f);
 		if (Handle != FKeyHandle::Invalid())
 		{
-			FTransaction Transaction = ConditionalTransaction(LOCTEXT("SetNamedCurveKey", "Setting Curve Key"), bShouldTransact);
+			CONDITIONAL_TRANSACTION(LOCTEXT("SetNamedCurveKey", "Setting Curve Key"));
 			// Cache old value for action
 			const FRichCurveKey CurrentKey = RichCurve->GetKey(Handle);
-			ConditionalAction<UE::Anim::FSetRichCurveKeyAction>(bShouldTransact, CurveId, CurrentKey);
+			CONDITIONAL_ACTION(UE::Anim::FSetRichCurveKeyAction, CurveId, CurrentKey);
 
 			// Set the new value
 			RichCurve->SetKeyValue(Handle, Key.Value);
 
-			Model->GetNotifier().Notify(EAnimDataModelNotifyType::CurveChanged, Payload);
+			Model->Notify(EAnimDataModelNotifyType::CurveChanged, Payload);
 		}
 		else
 		{
-			FTransaction Transaction = ConditionalTransaction(LOCTEXT("AddNamedCurveKey", "Adding Curve Key"), bShouldTransact);
-			ConditionalAction<UE::Anim::FRemoveRichCurveKeyAction>(bShouldTransact, CurveId, Key.Time);
+			CONDITIONAL_TRANSACTION(LOCTEXT("AddNamedCurveKey", "Adding Curve Key"));
+			CONDITIONAL_ACTION(UE::Anim::FRemoveRichCurveKeyAction, CurveId, Key.Time);
 
 			// Add the new key
 			RichCurve->AddKey(Key.Time, Key.Value);
 
-			Model->GetNotifier().Notify(EAnimDataModelNotifyType::CurveChanged, Payload);
+			Model->Notify(EAnimDataModelNotifyType::CurveChanged, Payload);
 		}
 
 		return true;
@@ -1157,7 +1101,9 @@ bool UAnimDataController::RemoveCurveKey(const FAnimationCurveIdentifier& CurveI
 {
 	ValidateModel();
 
-	if (FRichCurve* RichCurve = Model->GetMutableRichCurve(CurveId))
+	FRichCurve* RichCurve = Model->GetMutableRichCurve(CurveId);
+	ERawCurveTrackTypes SupportedCurveType = CurveId.CurveType;
+	if (RichCurve)
 	{
 		FCurveChangedPayload Payload;
 		Payload.Identifier = CurveId;
@@ -1166,15 +1112,15 @@ bool UAnimDataController::RemoveCurveKey(const FAnimationCurveIdentifier& CurveI
 		const FKeyHandle Handle = RichCurve->FindKey(Time, 0.f);
 		if (Handle != FKeyHandle::Invalid())
 		{
-			FTransaction Transaction = ConditionalTransaction(LOCTEXT("RemoveNamedCurveKey", "Removing Curve Key"), bShouldTransact);
+			CONDITIONAL_TRANSACTION(LOCTEXT("RemoveNamedCurveKey", "Removing Curve Key"));
 
 			// Cached current value for action
 			const FRichCurveKey CurrentKey = RichCurve->GetKey(Handle);
-			ConditionalAction<UE::Anim::FAddRichCurveKeyAction>(bShouldTransact, CurveId, CurrentKey);
+			CONDITIONAL_ACTION(UE::Anim::FAddRichCurveKeyAction, CurveId, CurrentKey);
 
 			RichCurve->DeleteKey(Handle);
 
-			Model->GetNotifier().Notify(EAnimDataModelNotifyType::CurveChanged, Payload);
+			Model->Notify(EAnimDataModelNotifyType::CurveChanged, Payload);
 
 			return true;
 		}
@@ -1192,17 +1138,19 @@ bool UAnimDataController::SetCurveKeys(const FAnimationCurveIdentifier& CurveId,
 {
 	ValidateModel();
 
-	if (FRichCurve* RichCurve = Model->GetMutableRichCurve(CurveId))
+	FRichCurve* RichCurve = Model->GetMutableRichCurve(CurveId);
+	ERawCurveTrackTypes SupportedCurveType = CurveId.CurveType;
+	if (RichCurve)
 	{
-		FTransaction Transaction = ConditionalTransaction(LOCTEXT("SettingNamedCurveKeys", "Setting Curve Keys"), bShouldTransact);
-		ConditionalAction<UE::Anim::FSetRichCurveKeysAction>(bShouldTransact, CurveId, RichCurve->GetConstRefOfKeys());
+		CONDITIONAL_TRANSACTION(LOCTEXT("SettingNamedCurveKeys", "Setting Curve Keys"));
+		CONDITIONAL_ACTION(UE::Anim::FSetRichCurveKeysAction, CurveId, RichCurve->GetConstRefOfKeys());
 
 		// Set rich curve values
 		RichCurve->SetKeys(CurveKeys);
 
 		FCurveChangedPayload Payload;
 		Payload.Identifier = CurveId;
-		Model->GetNotifier().Notify(EAnimDataModelNotifyType::CurveChanged, Payload);
+		Model->Notify(EAnimDataModelNotifyType::CurveChanged, Payload);
 
 		return true;
 	}
@@ -1213,42 +1161,165 @@ bool UAnimDataController::SetCurveKeys(const FAnimationCurveIdentifier& CurveId,
 void UAnimDataController::NotifyPopulated()
 {
 	ValidateModel();
-
-	Model->bPopulated = true;
-	Model->GetNotifier().Notify(EAnimDataModelNotifyType::Populated);
+	Model->Notify(EAnimDataModelNotifyType::Populated);
 }
 
 void UAnimDataController::NotifyBracketOpen()
 {
 	ValidateModel();
-	Model->GetNotifier().Notify(EAnimDataModelNotifyType::BracketOpened);
+	Model->Notify(EAnimDataModelNotifyType::BracketOpened);
 }
 
 void UAnimDataController::NotifyBracketClosed()
 {
 	ValidateModel();
-	Model->GetNotifier().Notify(EAnimDataModelNotifyType::BracketClosed);
+	Model->Notify(EAnimDataModelNotifyType::BracketClosed);
+}
+
+const bool UAnimDataController::IsSupportedCurveType(ERawCurveTrackTypes CurveType) const
+{
+	const TArray<ERawCurveTrackTypes> SupportedTypes = { ERawCurveTrackTypes::RCT_Float, ERawCurveTrackTypes::RCT_Transform };
+	return SupportedTypes.Contains(CurveType);
+}
+
+void UAnimDataController::ValidateModel() const
+{
+	checkf(Model != nullptr, TEXT("Invalid Model"));
 }
 
 void UAnimDataController::ResizePlayLength(float Length, float T0, float T1, bool bShouldTransact)
 {
-	ResizeNumberOfFrames(ConvertSecondsToFrameNumber(Length), ConvertSecondsToFrameNumber(T0), ConvertSecondsToFrameNumber(T1), bShouldTransact);
+	const TRange<float> PlayRange(TRange<float>::BoundsType::Inclusive(0.f), TRange<float>::BoundsType::Inclusive(Model->PlayLength));
+	if (!FMath::IsNearlyZero(Length) && Length > 0.f)
+	{
+		if (Length != Model->PlayLength)
+		{
+			// Ensure that T0 is within the curent play range
+			if (PlayRange.Contains(T0))
+			{
+				// Ensure that the start and end length of either removal or insertion are valid
+				if (T0 < T1)
+				{
+					CONDITIONAL_TRANSACTION(LOCTEXT("ResizePlayLength", "Resizing Play Length"));
+
+					FSequenceLengthChangedPayload Payload;
+					Payload.T0 = T0;
+					Payload.T1 = T1;
+					Payload.PreviousLength = Model->PlayLength;
+
+					CONDITIONAL_ACTION(UE::Anim::FResizePlayLengthAction, Model, T0, T1);
+
+					Model->PlayLength = Length;
+
+					Model->NumberOfFrames = Model->FrameRate.AsFrameTime(Model->PlayLength).RoundToFrame().Value;
+					Model->NumberOfKeys = Model->NumberOfFrames + 1;
+	
+					Model->Notify<FSequenceLengthChangedPayload>(EAnimDataModelNotifyType::SequenceLengthChanged, Payload);
+				}
+				else
+				{
+					ReportErrorf(LOCTEXT("InvalidEndTimeError", "Invalid T1, smaller that T0 value: T0 {0}, T1 {1}"), FText::AsNumber(T0), FText::AsNumber(T1));
+				}
+			}
+			else
+			{
+				ReportErrorf(LOCTEXT("InvalidStartTimeError", "Invalid T0, not within existing play range: T0 {0}, Play Length {1}"), FText::AsNumber(T0), FText::AsNumber(Model->PlayLength));
+			}
+		}
+		else
+		{
+			ReportWarningf(LOCTEXT("SamePlayLengthWarning", "New play length is same as existing one: {0} seconds"), FText::AsNumber(Length));
+		}
+	}
+	else
+	{
+		ReportErrorf(LOCTEXT("InvalidPlayLengthError", "Invalid play length value provided: {0} seconds"), FText::AsNumber(Length));
+	}
+}
+
+void UAnimDataController::ReportWarning(const FText& InMessage) const
+{
+	FString Message = InMessage.ToString();
+	if (Model != nullptr)
+	{
+		if (UPackage* Package = Cast<UPackage>(Model->GetOutermost()))
+		{
+			Message = FString::Printf(TEXT("%s : %s"), *Package->GetPathName(), *Message);
+		}
+	}
+
+	FScriptExceptionHandler::Get().HandleException(ELogVerbosity::Warning, *Message, *FString());
+}
+
+void UAnimDataController::ReportError(const FText& InMessage) const
+{
+	FString Message = InMessage.ToString();
+	if (Model != nullptr)
+	{
+		if (UPackage* Package = Cast<UPackage>(Model->GetOutermost()))
+		{
+			Message = FString::Printf(TEXT("%s : %s"), *Package->GetPathName(), *Message);
+		}
+	}
+
+	FScriptExceptionHandler::Get().HandleException(ELogVerbosity::Error, *Message, *FString());
+}
+
+FString UAnimDataController::GetCurveTypeValueName(ERawCurveTrackTypes InType) const
+{
+	FString ValueString;
+
+	const UEnum* Enum = FindObject<UEnum>(nullptr, TEXT("/Script/Engine.ERawCurveTrackTypes"));
+	if (Enum)
+	{
+		ValueString = Enum->GetNameStringByValue((int64)InType);
+	}
+
+	return ValueString;
+}
+
+bool UAnimDataController::CheckOuterClass(UClass* InClass) const
+{
+	ValidateModel();
+	
+	const UObject* ModelOuter = Model->GetOuter();
+	if (ModelOuter)
+	{
+		const UClass* OuterClass = ModelOuter->GetClass();
+		if (OuterClass)
+		{
+			if (OuterClass == InClass || OuterClass->IsChildOf(InClass))
+			{
+				return true;
+			}
+			else
+			{
+				ReportErrorf(LOCTEXT("NoValidOuterClassError", "Incorrect outer object class found for Animation Data Model {0}, expected {1} actual {2}"), FText::FromString(Model->GetName()), FText::FromString(InClass->GetName()), FText::FromString(OuterClass->GetName()));
+			}
+		}
+	}
+	else
+	{
+		ReportErrorf(LOCTEXT("NoValidOuterObjectFoundError", "No valid outer object found for Animation Data Model {0}"), FText::FromString(Model->GetName()));
+	}
+
+	return false;
 }
 
 int32 UAnimDataController::AddBoneTrack(FName BoneName, bool bShouldTransact /*= true*/)
 {
-	if (!ModelInterface->GetAnimationSequence())
+	if (!CheckOuterClass(UAnimSequence::StaticClass()))
 	{
 		return INDEX_NONE;
 	}
 
-	FTransaction Transaction = ConditionalTransaction(LOCTEXT("AddBoneTrack", "Adding Animation Data Track"), bShouldTransact);
+	CONDITIONAL_TRANSACTION(LOCTEXT("AddBoneTrack", "Adding Animation Data Track"));
 	return InsertBoneTrack(BoneName, INDEX_NONE, bShouldTransact);
 }
 
 int32 UAnimDataController::InsertBoneTrack(FName BoneName, int32 DesiredIndex, bool bShouldTransact /*= true*/)
 {
-	if (!ModelInterface->GetAnimationSequence())
+	if (!CheckOuterClass(UAnimSequence::StaticClass()))
 	{
 		return INDEX_NONE;
 	}
@@ -1260,26 +1331,29 @@ int32 UAnimDataController::InsertBoneTrack(FName BoneName, int32 DesiredIndex, b
 		if (Model->GetNumBoneTracks() >= MAX_ANIMATION_TRACKS)
 		{
 			ReportWarningf(LOCTEXT("MaxNumberOfTracksReachedWarning", "Cannot add track with name {0}. An animation sequence cannot contain more than 65535 tracks"), FText::FromName(BoneName));
-			return INDEX_NONE;
 		}
 		else
 		{
-			FTransaction Transaction = ConditionalTransaction(LOCTEXT("InsertBoneTrack", "Inserting Animation Data Track"), bShouldTransact);
+			CONDITIONAL_TRANSACTION(LOCTEXT("InsertBoneTrack", "Inserting Animation Data Track"));
 
 			// Determine correct index to do insertion at
 			const int32 InsertIndex = Model->BoneAnimationTracks.IsValidIndex(DesiredIndex) ? DesiredIndex : Model->BoneAnimationTracks.Num();
-			int32 BoneIndex = INDEX_NONE;
+
+			FBoneAnimationTrack& NewTrack = Model->BoneAnimationTracks.InsertDefaulted_GetRef(InsertIndex);
+			NewTrack.Name = BoneName;
 
 			if (const UAnimSequence* AnimationSequence = Model->GetAnimationSequence())
 			{
 				if (const USkeleton* Skeleton = AnimationSequence->GetSkeleton())
 				{
-					BoneIndex = Skeleton->GetReferenceSkeleton().FindBoneIndex(BoneName);
+					const int32 BoneIndex = Skeleton->GetReferenceSkeleton().FindBoneIndex(BoneName);
 
 					if (BoneIndex == INDEX_NONE)
 					{
-						ReportErrorf(LOCTEXT("UnableToFindBoneIndexWarning", "Unable to retrieve bone index for track: {0}"), FText::FromName(BoneName));
+						ReportWarningf(LOCTEXT("UnableToFindBoneIndexWarning", "Unable to retrieve bone index for track: {0}"), FText::FromName(BoneName));
 					}
+
+					NewTrack.BoneTreeIndex = BoneIndex;
 				}
 				else
 				{
@@ -1291,23 +1365,16 @@ int32 UAnimDataController::InsertBoneTrack(FName BoneName, int32 DesiredIndex, b
 				ReportError(LOCTEXT("UnableToGetOuterAnimSequenceError", "Unable to retrieve outer Animation Sequence"));
 			}
 
-			if (BoneIndex != INDEX_NONE)
-			{
-				FBoneAnimationTrack& NewTrack = Model->BoneAnimationTracks.InsertDefaulted_GetRef(InsertIndex);
-				NewTrack.Name = BoneName;
-				NewTrack.BoneTreeIndex = BoneIndex;
-
 				FAnimationTrackAddedPayload Payload;
 				Payload.Name = BoneName;
 				Payload.TrackIndex = InsertIndex;
 
-				Model->GetNotifier().Notify<FAnimationTrackAddedPayload>(EAnimDataModelNotifyType::TrackAdded, Payload);
-				ConditionalAction<UE::Anim::FRemoveTrackAction>(bShouldTransact, NewTrack);
+			Model->Notify<FAnimationTrackAddedPayload>(EAnimDataModelNotifyType::TrackAdded, Payload);
+			CONDITIONAL_ACTION(UE::Anim::FRemoveTrackAction, NewTrack, InsertIndex);
 
 				return InsertIndex;
 			}
 		}
-	}
 	else
 	{
 		ReportWarningf(LOCTEXT("TrackNameAlreadyExistsWarning", "Track with name {0} already exists"), FText::FromName(BoneName));
@@ -1318,7 +1385,7 @@ int32 UAnimDataController::InsertBoneTrack(FName BoneName, int32 DesiredIndex, b
 
 bool UAnimDataController::RemoveBoneTrack(FName BoneName, bool bShouldTransact /*= true*/)
 {
-	if (!ModelInterface->GetAnimationSequence())
+	if (!CheckOuterClass(UAnimSequence::StaticClass()))
 	{
 		return false;
 	}
@@ -1327,7 +1394,7 @@ bool UAnimDataController::RemoveBoneTrack(FName BoneName, bool bShouldTransact /
 
 	if (ExistingTrackPtr != nullptr)
 	{
-		FTransaction Transaction = ConditionalTransaction(LOCTEXT("RemoveBoneTrack", "Removing Animation Data Track"), bShouldTransact);
+		CONDITIONAL_TRANSACTION(LOCTEXT("RemoveBoneTrack", "Removing Animation Data Track"));
 		const int32 TrackIndex = Model->BoneAnimationTracks.IndexOfByPredicate([ExistingTrackPtr](const FBoneAnimationTrack& Track)
 		{
 			return Track.Name == ExistingTrackPtr->Name;
@@ -1335,13 +1402,13 @@ bool UAnimDataController::RemoveBoneTrack(FName BoneName, bool bShouldTransact /
 
 		ensure(TrackIndex != INDEX_NONE);
 
-		ConditionalAction<UE::Anim::FAddTrackAction>(bShouldTransact,*ExistingTrackPtr);
+		CONDITIONAL_ACTION(UE::Anim::FAddTrackAction, *ExistingTrackPtr, TrackIndex);
 		Model->BoneAnimationTracks.RemoveAt(TrackIndex);
 
 		FAnimationTrackRemovedPayload Payload;
 		Payload.Name = BoneName;
 
-		Model->GetNotifier().Notify(EAnimDataModelNotifyType::TrackRemoved, Payload);
+		Model->Notify(EAnimDataModelNotifyType::TrackRemoved, Payload);
 
 		return true;
 	}
@@ -1355,7 +1422,7 @@ bool UAnimDataController::RemoveBoneTrack(FName BoneName, bool bShouldTransact /
 
 void UAnimDataController::RemoveAllBoneTracks(bool bShouldTransact /*= true*/)
 {
-	if (!ModelInterface->GetAnimationSequence())
+	if (!CheckOuterClass(UAnimSequence::StaticClass()))
 	{
 		return;
 	}
@@ -1365,7 +1432,7 @@ void UAnimDataController::RemoveAllBoneTracks(bool bShouldTransact /*= true*/)
 
 	if (TrackNames.Num())
 	{
-		FBracket Bracket = ConditionalBracket(LOCTEXT("RemoveAllBoneTracks", "Removing all Animation Data Tracks"), bShouldTransact);
+		CONDITIONAL_BRACKET(LOCTEXT("RemoveAllBoneTracks", "Removing all Animation Data Tracks"));
 		for (const FName& TrackName : TrackNames)
 		{
 			RemoveBoneTrack(TrackName, bShouldTransact);
@@ -1375,12 +1442,12 @@ void UAnimDataController::RemoveAllBoneTracks(bool bShouldTransact /*= true*/)
 
 bool UAnimDataController::SetBoneTrackKeys(FName BoneName, const TArray<FVector>& PositionalKeys, const TArray<FQuat>& RotationalKeys, const TArray<FVector>& ScalingKeys, bool bShouldTransact /*= true*/)
 {
-	if (!ModelInterface->GetAnimationSequence())
+	if (!CheckOuterClass(UAnimSequence::StaticClass()))
 	{
 		return false;
 	}
 
-	FTransaction Transaction = ConditionalTransaction(LOCTEXT("SetTrackKeysTransaction", "Setting Animation Data Track keys"), bShouldTransact);
+	CONDITIONAL_TRANSACTION(LOCTEXT("SetTrackKeysTransaction", "Setting Animation Data Track keys"));
 
 	// Validate key format
 	const int32 MaxNumKeys = FMath::Max(FMath::Max(PositionalKeys.Num(), RotationalKeys.Num()), ScalingKeys.Num());
@@ -1395,7 +1462,7 @@ bool UAnimDataController::SetBoneTrackKeys(FName BoneName, const TArray<FVector>
 		{
 			if (FBoneAnimationTrack* TrackPtr = Model->FindMutableBoneTrackByName(BoneName))
 			{
-				ConditionalAction<UE::Anim::FSetTrackKeysAction>(bShouldTransact, *TrackPtr);
+				CONDITIONAL_ACTION(UE::Anim::FSetTrackKeysAction, *TrackPtr);
 
 
 				TrackPtr->InternalTrackData.PosKeys.SetNum(MaxNumKeys);
@@ -1411,7 +1478,7 @@ bool UAnimDataController::SetBoneTrackKeys(FName BoneName, const TArray<FVector>
 				FAnimationTrackChangedPayload Payload;
 				Payload.Name = BoneName;
 
-				Model->GetNotifier().Notify(EAnimDataModelNotifyType::TrackChanged, Payload);
+				Model->Notify(EAnimDataModelNotifyType::TrackChanged, Payload);
 
 				return true;
 			}
@@ -1440,7 +1507,7 @@ bool UAnimDataController::SetBoneTrackKeys(FName BoneName, const TArray<FVector3
 		return false;
 	}
 
-	FTransaction Transaction = ConditionalTransaction(LOCTEXT("SetTrackKeysTransaction", "Setting Animation Data Track keys"), bShouldTransact);
+	CONDITIONAL_TRANSACTION(LOCTEXT("SetTrackKeysTransaction", "Setting Animation Data Track keys"));
 
 	// Validate key format
 	const int32 MaxNumKeys = FMath::Max(FMath::Max(PositionalKeys.Num(), RotationalKeys.Num()), ScalingKeys.Num());
@@ -1455,7 +1522,7 @@ bool UAnimDataController::SetBoneTrackKeys(FName BoneName, const TArray<FVector3
 		{
 			if (FBoneAnimationTrack* TrackPtr = Model->FindMutableBoneTrackByName(BoneName))
 			{
-				ConditionalAction<UE::Anim::FSetTrackKeysAction>(bShouldTransact, *TrackPtr);
+				CONDITIONAL_ACTION(UE::Anim::FSetTrackKeysAction, *TrackPtr);
 
 				TrackPtr->InternalTrackData.PosKeys = PositionalKeys;
 				TrackPtr->InternalTrackData.RotKeys = RotationalKeys;
@@ -1464,7 +1531,7 @@ bool UAnimDataController::SetBoneTrackKeys(FName BoneName, const TArray<FVector3
 				FAnimationTrackChangedPayload Payload;
 				Payload.Name = BoneName;
 
-				Model->GetNotifier().Notify(EAnimDataModelNotifyType::TrackChanged, Payload);
+				Model->Notify(EAnimDataModelNotifyType::TrackChanged, Payload);
 
 				return true;
 			}
@@ -1486,23 +1553,23 @@ bool UAnimDataController::SetBoneTrackKeys(FName BoneName, const TArray<FVector3
 	return false;
 }
 
-static int32 DiscreteInclusiveLower(const TRange<int32>& InRange)
+int32 DiscreteInclusiveLower(const TRange<int32>& InRange)
 {
 	check(!InRange.GetLowerBound().IsOpen());
 
 	// Add one for exclusive lower bounds since they start on the next subsequent frame
-	static constexpr int32 Offsets[]   = { 0, 1 };
+	static const int32 Offsets[]   = { 0, 1 };
 	const int32        OffsetIndex = (int32)InRange.GetLowerBound().IsExclusive();
 
 	return InRange.GetLowerBound().GetValue() + Offsets[OffsetIndex];
 }
 
-static int32 DiscreteExclusiveUpper(const TRange<int32>& InRange)
+int32 DiscreteExclusiveUpper(const TRange<int32>& InRange)
 {
 	check(!InRange.GetUpperBound().IsOpen());
 
 	// Add one for inclusive upper bounds since they finish on the next subsequent frame
-	static constexpr int32 Offsets[]   = { 0, 1 };
+	static const int32 Offsets[]   = { 0, 1 };
 	const int32        OffsetIndex = (int32)InRange.GetUpperBound().IsInclusive();
 
 	return InRange.GetUpperBound().GetValue() + Offsets[OffsetIndex];
@@ -1515,7 +1582,7 @@ bool UAnimDataController::UpdateBoneTrackKeys(FName BoneName, const FInt32Range&
 		return false;
 	}
 
-	FTransaction Transaction = ConditionalTransaction(LOCTEXT("SetTrackKeysRangeTransaction", "Setting Animation Data Track keys"), bShouldTransact);
+	CONDITIONAL_TRANSACTION(LOCTEXT("SetTrackKeysRangeTransaction", "Setting Animation Data Track keys"));
 
 	// Validate key format
 	const int32 MaxNumKeys = FMath::Max(FMath::Max(PositionalKeys.Num(), RotationalKeys.Num()), ScalingKeys.Num());
@@ -1542,7 +1609,7 @@ bool UAnimDataController::UpdateBoneTrackKeys(FName BoneName, const FInt32Range&
 						TArray<FQuat4f>& TrackRotKeys = InternalTrackData.RotKeys;
 						TArray<FVector3f>& TrackScaleKeys = InternalTrackData.ScaleKeys;
 
-						ConditionalAction<UE::Anim::FSetTrackKeysAction>(bShouldTransact, *TrackPtr);
+						CONDITIONAL_ACTION(UE::Anim::FSetTrackKeysAction, *TrackPtr);
 
 						int32 KeyIndex = 0;
 						for (int32 FrameIndex = RangeMin; FrameIndex < RangeMax; ++FrameIndex, ++KeyIndex)
@@ -1555,7 +1622,7 @@ bool UAnimDataController::UpdateBoneTrackKeys(FName BoneName, const FInt32Range&
 						FAnimationTrackChangedPayload Payload;
 						Payload.Name = BoneName;
 
-						Model->GetNotifier().Notify(EAnimDataModelNotifyType::TrackChanged, Payload);
+						Model->Notify(EAnimDataModelNotifyType::TrackChanged, Payload);
 
 						return true;
 					}
@@ -1598,7 +1665,7 @@ bool UAnimDataController::UpdateBoneTrackKeys(FName BoneName, const FInt32Range&
 		return false;
 	}
 
-	FTransaction Transaction = ConditionalTransaction(LOCTEXT("SetTrackKeysTransaction", "Setting Animation Data Track keys"), bShouldTransact);
+	CONDITIONAL_TRANSACTION(LOCTEXT("SetTrackKeysTransaction", "Setting Animation Data Track keys"));
 
 	// Validate key format
 	const int32 MaxNumKeys = FMath::Max(FMath::Max(PositionalKeys.Num(), RotationalKeys.Num()), ScalingKeys.Num());
@@ -1625,7 +1692,7 @@ bool UAnimDataController::UpdateBoneTrackKeys(FName BoneName, const FInt32Range&
 						TArray<FQuat4f>& TrackRotKeys = InternalTrackData.RotKeys;
 						TArray<FVector3f>& TrackScaleKeys = InternalTrackData.ScaleKeys;
 
-						ConditionalAction<UE::Anim::FSetTrackKeysAction>(bShouldTransact, *TrackPtr);
+						CONDITIONAL_ACTION(UE::Anim::FSetTrackKeysAction, *TrackPtr);
 
 						int32 KeyIndex = 0;
 						for (int32 FrameIndex = RangeMin; FrameIndex < RangeMax; ++FrameIndex, ++KeyIndex)
@@ -1638,7 +1705,7 @@ bool UAnimDataController::UpdateBoneTrackKeys(FName BoneName, const FInt32Range&
 						FAnimationTrackChangedPayload Payload;
 						Payload.Name = BoneName;
 
-						Model->GetNotifier().Notify(EAnimDataModelNotifyType::TrackChanged, Payload);
+						Model->Notify(EAnimDataModelNotifyType::TrackChanged, Payload);
 
 						return true;
 					}
@@ -1676,7 +1743,7 @@ bool UAnimDataController::UpdateBoneTrackKeys(FName BoneName, const FInt32Range&
 
 void UAnimDataController::ResizeCurves(float NewLength, bool bInserted, float T0, float T1, bool bShouldTransact /*= true*/)
 {
-	FBracket Bracket = ConditionalBracket(LOCTEXT("ResizeCurves", "Resizing all Curves"), bShouldTransact);
+	CONDITIONAL_BRACKET(LOCTEXT("ResizeCurves", "Resizing all Curves"));
 
 	for (FFloatCurve& Curve : Model->CurveData.FloatCurves)
 	{
@@ -1690,11 +1757,11 @@ void UAnimDataController::ResizeCurves(float NewLength, bool bInserted, float T0
 		FTransformCurve ResizedCurve = Curve;
 		for (int32 SubCurveIndex = 0; SubCurveIndex < 3; ++SubCurveIndex)
 		{
-			const ETransformCurveChannel Channel = static_cast<ETransformCurveChannel>(SubCurveIndex);
+			const ETransformCurveChannel Channel = (ETransformCurveChannel)SubCurveIndex;
 			FVectorCurve& SubCurve = *ResizedCurve.GetVectorCurveByIndex(SubCurveIndex);
 			for (int32 ChannelIndex = 0; ChannelIndex < 3; ++ChannelIndex)
 			{
-				const EVectorCurveChannel Axis = static_cast<EVectorCurveChannel>(ChannelIndex);
+				const EVectorCurveChannel Axis = (EVectorCurveChannel)ChannelIndex;
 				FAnimationCurveIdentifier TargetCurveIdentifier = FAnimationCurveIdentifier(Curve.Name, ERawCurveTrackTypes::RCT_Transform);
 				UAnimationCurveIdentifierExtensions::GetTransformChildCurveIdentifier(TargetCurveIdentifier, Channel, Axis);
 				
@@ -1708,7 +1775,7 @@ void UAnimDataController::ResizeCurves(float NewLength, bool bInserted, float T0
 
 void UAnimDataController::ResizeAttributes(float NewLength, bool bInserted, float T0, float T1, bool bShouldTransact)
 {
-	FBracket Bracket = ConditionalBracket(LOCTEXT("ResizeAttributes", "Resizing all Attributes"), bShouldTransact);
+	CONDITIONAL_BRACKET(LOCTEXT("ResizeAttributes", "Resizing all Attributes"));
 
 	for (FAnimatedBoneAttribute& Attribute : Model->AnimatedBoneAttributes)
 	{
@@ -1740,18 +1807,18 @@ bool UAnimDataController::AddAttribute(const FAnimationAttributeIdentifier& Attr
 
 		if (!bAttributeAlreadyExists)
 		{
-			FTransaction Transaction = ConditionalTransaction(LOCTEXT("AddAttribute", "Adding Animated Bone Attribute"), bShouldTransact);
+			CONDITIONAL_TRANSACTION(LOCTEXT("AddAttribute", "Adding Animated Bone Attribute"));
 
 			FAnimatedBoneAttribute& Attribute = Model->AnimatedBoneAttributes.AddDefaulted_GetRef();
 			Attribute.Identifier = AttributeIdentifier;
 
 			Attribute.Curve.SetScriptStruct(AttributeIdentifier.GetType());
 		
-			ConditionalAction<UE::Anim::FRemoveAtributeAction>(bShouldTransact, AttributeIdentifier);
+			CONDITIONAL_ACTION(UE::Anim::FRemoveAtributeAction, AttributeIdentifier);
 
 			FAttributeAddedPayload Payload;
 			Payload.Identifier = AttributeIdentifier;
-			Model->GetNotifier().Notify(EAnimDataModelNotifyType::AttributeAdded, Payload);
+			Model->Notify(EAnimDataModelNotifyType::AttributeAdded, Payload);
 
 			return true;
 		}
@@ -1780,15 +1847,15 @@ bool UAnimDataController::RemoveAttribute(const FAnimationAttributeIdentifier& A
 
 		if (AttributeIndex != INDEX_NONE)
 		{
-			FTransaction Transaction = ConditionalTransaction(LOCTEXT("RemoveAttribute", "Removing Animated Bone Attribute"), bShouldTransact);
+			CONDITIONAL_TRANSACTION(LOCTEXT("RemoveAttribute", "Removing Animated Bone Attribute"));
 
-			ConditionalAction<UE::Anim::FAddAtributeAction>(bShouldTransact, Model->AnimatedBoneAttributes[AttributeIndex]);
+			CONDITIONAL_ACTION(UE::Anim::FAddAtributeAction, Model->AnimatedBoneAttributes[AttributeIndex]);
 
 			Model->AnimatedBoneAttributes.RemoveAtSwap(AttributeIndex);
 			
 			FAttributeRemovedPayload Payload;
 			Payload.Identifier = AttributeIdentifier;
-			Model->GetNotifier().Notify(EAnimDataModelNotifyType::AttributeRemoved, Payload);
+			Model->Notify(EAnimDataModelNotifyType::AttributeRemoved, Payload);
 
 			return true;
 		}
@@ -1825,7 +1892,7 @@ int32 UAnimDataController::RemoveAllAttributesForBone(const FName& BoneName, boo
 
 	if (Identifiers.Num())
 	{
-		FBracket Bracket = ConditionalBracket(LOCTEXT("RemoveAllAttributesForBone", "Removing all Attributes for Bone"), bShouldTransact);
+		CONDITIONAL_BRACKET(LOCTEXT("RemoveAllAttributesForBone", "Removing all Attributes for Bone"));
 		for (const FAnimationAttributeIdentifier& Identifier : Identifiers)
 		{
 			NumRemovedAttributes += RemoveAttribute(Identifier, bShouldTransact) ? 1 : 0;
@@ -1848,7 +1915,7 @@ int32 UAnimDataController::RemoveAllAttributes(bool bShouldTransact)
 
 	if (Identifiers.Num())
 	{
-		FBracket Bracket = ConditionalBracket(LOCTEXT("RemoveAllAttributes", "Removing all Attributes"), bShouldTransact);
+		CONDITIONAL_BRACKET(LOCTEXT("RemoveAllAttributes", "Removing all Attributes"));
 		for (const FAnimationAttributeIdentifier& Identifier : Identifiers)
 		{
 			NumRemovedAttributes += RemoveAttribute(Identifier, bShouldTransact) ? 1: 0;
@@ -1864,7 +1931,7 @@ bool UAnimDataController::SetAttributeKey_Internal(const FAnimationAttributeIden
 	{
 		if (KeyValue)
 		{
-			FAnimatedBoneAttribute* AttributePtr = Model->AnimatedBoneAttributes.FindByPredicate([AttributeIdentifier](const FAnimatedBoneAttribute& Attribute)
+			FAnimatedBoneAttribute* AttributePtr = Model->AnimatedBoneAttributes.FindByPredicate([AttributeIdentifier](FAnimatedBoneAttribute& Attribute)
 			{
 				return Attribute.Identifier == AttributeIdentifier;
 			});
@@ -1873,26 +1940,26 @@ bool UAnimDataController::SetAttributeKey_Internal(const FAnimationAttributeIden
 			{
 				if (TypeStruct == AttributePtr->Identifier.GetType())
 				{
-					FTransaction Transaction = ConditionalTransaction(LOCTEXT("SettingAttributeKey", "Setting Animated Bone Attribute key"), bShouldTransact);
+					CONDITIONAL_TRANSACTION(LOCTEXT("SettingAttributeKey", "Setting Animated Bone Attribute key"));
 
 					FAttributeCurve& Curve = AttributePtr->Curve;
-					const FKeyHandle KeyHandle = Curve.FindKey(Time);
+					FKeyHandle KeyHandle = Curve.FindKey(Time);
 					// In case the key does not yet exist one will be added, and thus the undo is a remove
 					if (KeyHandle == FKeyHandle::Invalid())
 					{
-						ConditionalAction<UE::Anim::FRemoveAtributeKeyAction>(bShouldTransact, AttributeIdentifier, Time);
-						Curve.UpdateOrAddTypedKey(Time, KeyValue, TypeStruct);
+						CONDITIONAL_ACTION(UE::Anim::FRemoveAtributeKeyAction, AttributeIdentifier, Time);
+						Curve.UpdateOrAddKey(Time, KeyValue);
 					}
 					// In case the key does exist it will be updated , and thus the undo is a revert to the current value
 					else
 					{
-						ConditionalAction<UE::Anim::FSetAtributeKeyAction>(bShouldTransact, AttributeIdentifier, Curve.GetKey(KeyHandle));
-						Curve.UpdateOrAddTypedKey(Time, KeyValue, TypeStruct);
+						CONDITIONAL_ACTION(UE::Anim::FSetAtributeKeyAction, AttributeIdentifier, Curve.GetKey(KeyHandle));
+						Curve.UpdateOrAddKey(Time, KeyValue);
 					}
 
 					FAttributeChangedPayload Payload;
 					Payload.Identifier = AttributeIdentifier;
-					Model->GetNotifier().Notify(EAnimDataModelNotifyType::AttributeChanged, Payload);
+					Model->Notify(EAnimDataModelNotifyType::AttributeChanged, Payload);
 
 					return true;
 				}
@@ -1927,7 +1994,7 @@ bool UAnimDataController::SetAttributeKeys_Internal(const FAnimationAttributeIde
 	{
 		if (Times.Num() == KeyValues.Num())
 		{
-			FAnimatedBoneAttribute* AttributePtr = Model->AnimatedBoneAttributes.FindByPredicate([AttributeIdentifier](const FAnimatedBoneAttribute& Attribute)
+			FAnimatedBoneAttribute* AttributePtr = Model->AnimatedBoneAttributes.FindByPredicate([AttributeIdentifier](FAnimatedBoneAttribute& Attribute)
 			{
 				return Attribute.Identifier == AttributeIdentifier;
 			});
@@ -1936,17 +2003,17 @@ bool UAnimDataController::SetAttributeKeys_Internal(const FAnimationAttributeIde
 			{
 				if (TypeStruct == AttributePtr->Identifier.GetType())
 				{
-					FTransaction Transaction = ConditionalTransaction(LOCTEXT("SettingAttributeKeys", "Setting Animated Bone Attribute keys"), bShouldTransact);
+					CONDITIONAL_TRANSACTION(LOCTEXT("SettingAttributeKeys", "Setting Animated Bone Attribute keys"));
 
 					FAnimatedBoneAttribute& Attribute = *AttributePtr;
 
-					ConditionalAction<UE::Anim::FSetAtributeKeysAction>(bShouldTransact, Attribute);
+					CONDITIONAL_ACTION(UE::Anim::FSetAtributeKeysAction, Attribute);
 			
 					Attribute.Curve.SetKeys(Times, KeyValues);
 
 					FAttributeChangedPayload Payload;
 					Payload.Identifier = AttributeIdentifier;
-					Model->GetNotifier().Notify(EAnimDataModelNotifyType::AttributeChanged, Payload);
+					Model->Notify(EAnimDataModelNotifyType::AttributeChanged, Payload);
 
 					return true;
 				}
@@ -1989,19 +2056,19 @@ bool UAnimDataController::RemoveAttributeKey(const FAnimationAttributeIdentifier
 		if (AttributePtr)
 		{
 			FAttributeCurve& Curve = AttributePtr->Curve;
-			const FKeyHandle KeyHandle = Curve.FindKey(Time);
+			FKeyHandle KeyHandle = Curve.FindKey(Time);
 
 			if (KeyHandle != FKeyHandle::Invalid())
 			{
-				FTransaction Transaction = ConditionalTransaction(LOCTEXT("RemovingAttributeKey", "Removing Animated Bone Attribute key"), bShouldTransact);
+				CONDITIONAL_TRANSACTION(LOCTEXT("RemovingAttributeKey", "Removing Animated Bone Attribute key"));
 
-				ConditionalAction<UE::Anim::FAddAtributeKeyAction>(bShouldTransact, AttributeIdentifier, Curve.GetKey(KeyHandle));
+				CONDITIONAL_ACTION(UE::Anim::FAddAtributeKeyAction, AttributeIdentifier, Curve.GetKey(KeyHandle));
 
 				Curve.DeleteKey(KeyHandle);
 
 				FAttributeAddedPayload Payload;
 				Payload.Identifier = AttributeIdentifier;
-				Model->GetNotifier().Notify(EAnimDataModelNotifyType::AttributeChanged, Payload);
+				Model->Notify(EAnimDataModelNotifyType::AttributeChanged, Payload);
 
 				return true;
 			}
@@ -2037,7 +2104,7 @@ bool UAnimDataController::DuplicateAttribute(const FAnimationAttributeIdentifier
 			{
 				if(const FAnimatedBoneAttribute* AttributePtr = Model->FindAttribute(AttributeIdentifier))
 				{
-					FTransaction Transaction = ConditionalTransaction(LOCTEXT("DuplicateAttribute", "Duplicating Animation Attribute"), bShouldTransact);
+					CONDITIONAL_TRANSACTION(LOCTEXT("DuplicateAttribute", "Duplicating Animation Attribute"));
 
 					FAnimatedBoneAttribute& DuplicateAttribute = Model->AnimatedBoneAttributes.AddDefaulted_GetRef();
 					DuplicateAttribute.Identifier = NewAttributeIdentifier;
@@ -2045,9 +2112,9 @@ bool UAnimDataController::DuplicateAttribute(const FAnimationAttributeIdentifier
 
 					FAttributeAddedPayload Payload;
 					Payload.Identifier = NewAttributeIdentifier;
-					Model->GetNotifier().Notify(EAnimDataModelNotifyType::AttributeAdded, Payload);
+					Model->Notify(EAnimDataModelNotifyType::AttributeAdded, Payload);
 
-					ConditionalAction<UE::Anim::FRemoveAtributeAction>(bShouldTransact, NewAttributeIdentifier);
+					CONDITIONAL_ACTION(UE::Anim::FRemoveAtributeAction, NewAttributeIdentifier);
 					
 					return true;
 				}
@@ -2074,28 +2141,6 @@ bool UAnimDataController::DuplicateAttribute(const FAnimationAttributeIdentifier
 	return false;
 }
 
-void UAnimDataController::UpdateWithSkeleton(USkeleton* TargetSkeleton, bool bShouldTransact)
-{
-	RemoveBoneTracksMissingFromSkeleton(TargetSkeleton);
-}
-
-void UAnimDataController::PopulateWithExistingModel(TScriptInterface<IAnimationDataModel> InModel)
-{	
-	Model->BoneAnimationTracks = InModel->GetBoneAnimationTracks();
-	Model->FrameRate = InModel->GetFrameRate();
-	Model->NumberOfFrames = InModel->GetNumberOfFrames();
-	Model->NumberOfKeys = InModel->GetNumberOfFrames() + 1;
-	Model->CurveData = InModel->GetCurveData();
-	Model->AnimatedBoneAttributes = InModel->GetAttributes();
-
-	// Ensure curve name UIDs are validated against the outer skeleton as they are not copied over through FSmartName
-	if (const UAnimSequenceBase* AnimSequenceBase = Cast<UAnimSequenceBase>(InModel.GetObject()->GetOuter()))
-	{
-		FindOrAddCurveNamesOnSkeleton(AnimSequenceBase->GetSkeleton(), ERawCurveTrackTypes::RCT_Float);
-		FindOrAddCurveNamesOnSkeleton(AnimSequenceBase->GetSkeleton(), ERawCurveTrackTypes::RCT_Transform);
-	}
-}
-
 #endif // WITH_EDITOR
 
 #undef LOCTEXT_NAMESPACE // "AnimDataController"
