[Backout] - CL20593017, 20594056, 20625634, 20645670, 20649179, 20649223, 20649380, 20658535

diff --git a/Engine/Source/Developer/AnimationDataController/Private/AnimDataControllerActions.cpp b/Engine/Source/Developer/AnimationDataController/Private/AnimDataControllerActions.cpp
--- a/Engine/Source/Developer/AnimationDataController/Private/AnimDataControllerActions.cpp
+++ b/Engine/Source/Developer/AnimationDataController/Private/AnimDataControllerActions.cpp
@@ -3,7 +3,7 @@
 #include "AnimDataControllerActions.h"
 #include "AnimDataController.h"
 
-#include "Animation/AnimData/IAnimationDataModel.h"
+#include "Animation/AnimData/AnimDataModel.h"
 #include "UObject/StrongObjectPtr.h"
 
 #define LOCTEXT_NAMESPACE "AnimDataControllerActions"
@@ -16,13 +16,13 @@ namespace Anim {
 
 TUniquePtr<FChange> FAnimDataBaseAction::Execute(UObject* Object)
 {
-	const TScriptInterface<IAnimationDataModel> DataModelInterface(Object);
-	checkf(DataModelInterface, TEXT("Invalid IAnimationDataModel Object"));
+	UAnimDataModel* Model = Cast<UAnimDataModel>(Object);
+	checkf(Model != nullptr, TEXT("Invalid UAnimDataModel Object"));
 
-	const TScriptInterface<IAnimationDataController> Controller = DataModelInterface->GetController();
-	Controller->SetModel(DataModelInterface);
+	TStrongObjectPtr<UAnimDataController> Controller(NewObject<UAnimDataController>());
+	Controller->SetModel(Model);
 
-	return ExecuteInternal(DataModelInterface.GetInterface(), Controller.GetInterface());
+	return ExecuteInternal(Model, Controller.Get());
 }
 
 FString FAnimDataBaseAction::ToString() const
@@ -30,7 +30,7 @@ FString FAnimDataBaseAction::ToString() const
 	return ToStringInternal();
 }
 
-TUniquePtr<FChange> FOpenBracketAction::ExecuteInternal(IAnimationDataModel* Model, IAnimationDataController* Controller)
+TUniquePtr<FChange> FOpenBracketAction::ExecuteInternal(UAnimDataModel* Model, UAnimDataController* Controller)
 {
 	Controller->NotifyBracketOpen();
 	return MakeUnique<FCloseBracketAction>(Description);
@@ -41,7 +41,7 @@ FString FOpenBracketAction::ToStringInternal() const
 	return FString::Printf(TEXT("Open Bracket: %s"), *Description);
 }
 
-TUniquePtr<FChange> FCloseBracketAction::ExecuteInternal(IAnimationDataModel* Model, IAnimationDataController* Controller)
+TUniquePtr<FChange> FCloseBracketAction::ExecuteInternal(UAnimDataModel* Model, UAnimDataController* Controller)
 {
 	Controller->NotifyBracketClosed();
 	return MakeUnique<FOpenBracketAction>(Description);
@@ -52,16 +52,17 @@ FString FCloseBracketAction::ToStringInternal() const
 	return TEXT("Closing Bracket");
 }
 
-FAddTrackAction::FAddTrackAction(const FBoneAnimationTrack& Track)
+FAddTrackAction::FAddTrackAction(const FBoneAnimationTrack& Track, int32 InTrackIndex)
 {
 	Name = Track.Name;
 	BoneTreeIndex = Track.BoneTreeIndex;
 	Data = Track.InternalTrackData;
+	TrackIndex = InTrackIndex;
 }
 
-TUniquePtr<FChange> FAddTrackAction::ExecuteInternal(IAnimationDataModel* Model, IAnimationDataController* Controller)
+TUniquePtr<FChange> FAddTrackAction::ExecuteInternal(UAnimDataModel* Model, UAnimDataController* Controller)
 {
-	Controller->AddBoneTrack(Name, false);
+	Controller->InsertBoneTrack(Name, TrackIndex, false);
 	
 	if (Data.PosKeys.Num() || Data.RotKeys.Num() || Data.ScaleKeys.Num())
 	{
@@ -69,7 +70,7 @@ TUniquePtr<FChange> FAddTrackAction::ExecuteInternal(IAnimationDataModel* Model,
 	}
 	
 	const FBoneAnimationTrack& AddedTrack = Model->GetBoneTrackByName(Name);
-	return MakeUnique<FRemoveTrackAction>(AddedTrack);
+	return MakeUnique<FRemoveTrackAction>(AddedTrack, TrackIndex);
 }
 
 FString FAddTrackAction::ToStringInternal() const
@@ -77,16 +78,17 @@ FString FAddTrackAction::ToStringInternal() const
 	return FText::Format(LOCTEXT("AddTrackAction_Description", "Adding animation bone track '{0}'."), FText::FromName(Name)).ToString();
 }
 
-FRemoveTrackAction::FRemoveTrackAction(const FBoneAnimationTrack& Track)
+FRemoveTrackAction::FRemoveTrackAction(const FBoneAnimationTrack& Track, int32 InTrackIndex)
 {
 	Name = Track.Name;
+	TrackIndex = InTrackIndex;
 }
 
-TUniquePtr<FChange> FRemoveTrackAction::ExecuteInternal(IAnimationDataModel* Model, IAnimationDataController* Controller)
+TUniquePtr<FChange> FRemoveTrackAction::ExecuteInternal(UAnimDataModel* Model, UAnimDataController* Controller)
 {
 	const FBoneAnimationTrack& Track = Model->GetBoneTrackByName(Name);
 
-	TUniquePtr<FChange> InverseAction = MakeUnique<FAddTrackAction>(Track);
+	TUniquePtr<FChange> InverseAction = MakeUnique<FAddTrackAction>(Track, TrackIndex);
 	Controller->RemoveBoneTrack(Name, false);
 	return InverseAction;
 }
@@ -103,7 +105,7 @@ FSetTrackKeysAction::FSetTrackKeysAction(const FBoneAnimationTrack& Track)
 	TrackData = Track.InternalTrackData;
 }
 
-TUniquePtr<FChange> FSetTrackKeysAction::ExecuteInternal(IAnimationDataModel* Model, IAnimationDataController* Controller)
+TUniquePtr<FChange> FSetTrackKeysAction::ExecuteInternal(UAnimDataModel* Model, UAnimDataController* Controller)
 {
 	const FBoneAnimationTrack& Track = Model->GetBoneTrackByName(Name);
 	TUniquePtr<FChange> InverseAction = MakeUnique<FSetTrackKeysAction>(Track);
@@ -117,29 +119,29 @@ FString FSetTrackKeysAction::ToStringInternal() const
 	return FText::Format(LOCTEXT("SetTrackKeysAction_Description", "Setting keys for animation bone track '{0}'."), FText::FromName(Name)).ToString();
 }
 
-FResizePlayLengthInFramesAction::FResizePlayLengthInFramesAction(const IAnimationDataModel* InModel, FFrameNumber F0, FFrameNumber F1) : Frame0(F0), Frame1(F1)
+FResizePlayLengthAction::FResizePlayLengthAction(const UAnimDataModel* InModel, float t0, float t1) : T0(t0), T1(t1)
 {
-	Length = InModel->GetNumberOfFrames();
+	Length = InModel->GetPlayLength();
 }
 
-TUniquePtr<FChange> FResizePlayLengthInFramesAction::ExecuteInternal(IAnimationDataModel* Model, IAnimationDataController* Controller)
+TUniquePtr<FChange> FResizePlayLengthAction::ExecuteInternal(UAnimDataModel* Model, UAnimDataController* Controller)
 {
-	TUniquePtr<FChange> InverseAction = MakeUnique<FResizePlayLengthInFramesAction>(Model, Frame0, Frame1);
-	Controller->ResizeNumberOfFrames(Length, Frame0, Frame1, false);
+	TUniquePtr<FChange> InverseAction = MakeUnique<FResizePlayLengthAction>(Model, T0, T1);
+	Controller->ResizePlayLength(Length, T0, T1, false);
 	return InverseAction;
 }
 
-FString FResizePlayLengthInFramesAction::ToStringInternal() const
+FString FResizePlayLengthAction::ToStringInternal() const
 {
-	return FText::Format(LOCTEXT("ResizePlayLengthInFramesAction_Description", "Resizing play length to {0} frames."), FText::AsNumber(Length.Value)).ToString();
+	return FText::Format(LOCTEXT("ResizePlayLengthAction_Description", "Resizing play length to {0}."), FText::AsNumber(Length)).ToString();
 }
 
-FSetFrameRateAction::FSetFrameRateAction(const IAnimationDataModel* InModel)
+FSetFrameRateAction::FSetFrameRateAction(const UAnimDataModel* InModel)
 {
 	FrameRate = InModel->GetFrameRate();
 }
 
-TUniquePtr<FChange> FSetFrameRateAction::ExecuteInternal(IAnimationDataModel* Model, IAnimationDataController* Controller)
+TUniquePtr<FChange> FSetFrameRateAction::ExecuteInternal(UAnimDataModel* Model, UAnimDataController* Controller)
 {
 	TUniquePtr<FChange> InverseAction = MakeUnique<FSetFrameRateAction>(Model);
 	Controller->SetFrameRate(FrameRate, false);
@@ -151,7 +153,7 @@ FString FSetFrameRateAction::ToStringInternal() const
 	return FText::Format(LOCTEXT("SetFrameRateAction_Description", "Setting Frame Rate to {0}."), FrameRate.ToPrettyText()).ToString();
 }
 
-TUniquePtr<FChange> FAddCurveAction::ExecuteInternal(IAnimationDataModel* Model, IAnimationDataController* Controller)
+TUniquePtr<FChange> FAddCurveAction::ExecuteInternal(UAnimDataModel* Model, UAnimDataController* Controller)
 {
 	TUniquePtr<FChange> InverseAction = MakeUnique<FRemoveCurveAction>(CurveId);
 
@@ -168,7 +170,7 @@ FString FAddCurveAction::ToStringInternal() const
 	return FText::Format(LOCTEXT("AddCurveAction_Description", "Adding {0} curve '{1}'."), FText::FromString(CurveId.CurveType == ERawCurveTrackTypes::RCT_Float ? FloatLabel : TransformLabel), FText::FromName(CurveId.InternalName.DisplayName)).ToString();
 }
 
-TUniquePtr<FChange> FRemoveCurveAction::ExecuteInternal(IAnimationDataModel* Model, IAnimationDataController* Controller)
+TUniquePtr<FChange> FRemoveCurveAction::ExecuteInternal(UAnimDataModel* Model, UAnimDataController* Controller)
 {
 	TUniquePtr<FChange> InverseAction;	
 
@@ -196,7 +198,7 @@ FString FRemoveCurveAction::ToStringInternal() const
 	return FText::Format(LOCTEXT("RemoveCurveAction_Description", "Removing {0} curve '{1}'."), FText::FromString(CurveId.CurveType == ERawCurveTrackTypes::RCT_Float ? FloatLabel : TransformLabel), FText::FromName(CurveId.InternalName.DisplayName)).ToString();
 }
 
-TUniquePtr<FChange> FSetCurveFlagsAction::ExecuteInternal(IAnimationDataModel* Model, IAnimationDataController* Controller)
+TUniquePtr<FChange> FSetCurveFlagsAction::ExecuteInternal(UAnimDataModel* Model, UAnimDataController* Controller)
 {
 	const FAnimCurveBase& Curve = Model->GetCurve(CurveId);
 
@@ -216,7 +218,7 @@ FString FSetCurveFlagsAction::ToStringInternal() const
 	return FText::Format(LOCTEXT("SetCurveFlagsAction_Description", "Setting flags for {0} curve '{1}'."), FText::FromString(CurveType == ERawCurveTrackTypes::RCT_Float ? FloatLabel : TransformLabel), FText::FromName(CurveId.InternalName.DisplayName)).ToString();
 }
 
-TUniquePtr<FChange> FRenameCurveAction::ExecuteInternal(IAnimationDataModel* Model, IAnimationDataController* Controller)
+TUniquePtr<FChange> FRenameCurveAction::ExecuteInternal(UAnimDataModel* Model, UAnimDataController* Controller)
 {
 	Controller->RenameCurve(CurveId, NewCurveId, false);
 	return MakeUnique<FRenameCurveAction>(NewCurveId, CurveId);
@@ -234,7 +236,7 @@ FScaleCurveAction::FScaleCurveAction(const FAnimationCurveIdentifier& InCurveId,
 	ensure(CurveType == ERawCurveTrackTypes::RCT_Float);
 }
 
-TUniquePtr<FChange> FScaleCurveAction::ExecuteInternal(IAnimationDataModel* Model, IAnimationDataController* Controller)
+TUniquePtr<FChange> FScaleCurveAction::ExecuteInternal(UAnimDataModel* Model, UAnimDataController* Controller)
 {
 	const float InverseFactor = 1.0f / Factor;
 
@@ -253,7 +255,7 @@ FString FScaleCurveAction::ToStringInternal() const
 	return FText::Format(LOCTEXT("ScaleCurveAction_Description", "Scaling {0} curve '{1}'."), FText::FromString(CurveType == ERawCurveTrackTypes::RCT_Float ? FloatLabel : TransformLabel), FText::FromName(CurveId.InternalName.DisplayName)).ToString();
 }
 
-TUniquePtr<FChange> FAddFloatCurveAction::ExecuteInternal(IAnimationDataModel* Model, IAnimationDataController* Controller)
+TUniquePtr<FChange> FAddFloatCurveAction::ExecuteInternal(UAnimDataModel* Model, UAnimDataController* Controller)
 {
 	TUniquePtr<FChange> InverseAction = MakeUnique<FRemoveCurveAction>(CurveId);
 
@@ -286,7 +288,7 @@ FAddTransformCurveAction::FAddTransformCurveAction(const FAnimationCurveIdentifi
 	SubCurveKeys[8] = InTransformCurve.ScaleCurve.FloatCurves[2].GetConstRefOfKeys();
 }
 
-TUniquePtr<FChange> FAddTransformCurveAction::ExecuteInternal(IAnimationDataModel* Model, IAnimationDataController* Controller)
+TUniquePtr<FChange> FAddTransformCurveAction::ExecuteInternal(UAnimDataModel* Model, UAnimDataController* Controller)
 {
 	TUniquePtr<FChange> InverseAction = MakeUnique<FRemoveCurveAction>(CurveId);
 
@@ -294,10 +296,10 @@ TUniquePtr<FChange> FAddTransformCurveAction::ExecuteInternal(IAnimationDataMode
 
 	for (int32 SubCurveIndex = 0; SubCurveIndex < 3; ++SubCurveIndex)
 	{
-		const ETransformCurveChannel Channel = static_cast<ETransformCurveChannel>(SubCurveIndex);
+		const ETransformCurveChannel Channel = (ETransformCurveChannel)SubCurveIndex;
 		for (int32 ChannelIndex = 0; ChannelIndex < 3; ++ChannelIndex)
 		{
-			const EVectorCurveChannel Axis = static_cast<EVectorCurveChannel>(ChannelIndex);
+			const EVectorCurveChannel Axis = (EVectorCurveChannel)ChannelIndex;
 			FAnimationCurveIdentifier TargetCurveIdentifier = CurveId;
 			UAnimationCurveIdentifierExtensions::GetTransformChildCurveIdentifier(TargetCurveIdentifier, Channel, Axis);
 			Controller->SetCurveKeys(TargetCurveIdentifier, SubCurveKeys[(SubCurveIndex * 3) + ChannelIndex], false);
@@ -312,7 +314,7 @@ FString FAddTransformCurveAction::ToStringInternal() const
 	return FText::Format(LOCTEXT("AddTransformCurveAction_Description", "Adding transform curve '{0}'."), FText::FromName(CurveId.InternalName.DisplayName)).ToString();
 }
 
-TUniquePtr<FChange> FAddRichCurveKeyAction::ExecuteInternal(IAnimationDataModel* Model, IAnimationDataController* Controller)
+TUniquePtr<FChange> FAddRichCurveKeyAction::ExecuteInternal(UAnimDataModel* Model, UAnimDataController* Controller)
 {
 	TUniquePtr<FChange> InverseAction = MakeUnique<FRemoveRichCurveKeyAction>(CurveId, Key.Time);
 	Controller->SetCurveKey(CurveId, Key, false);
@@ -327,7 +329,7 @@ FString FAddRichCurveKeyAction::ToStringInternal() const
 	return FText::Format(LOCTEXT("AddNamedRichCurveKeyAction_Description", "Adding key to {0} curve '{1}'."), FText::FromString(CurveId.CurveType == ERawCurveTrackTypes::RCT_Float ? FloatLabel : TransformLabel), FText::FromName(CurveId.InternalName.DisplayName)).ToString();
 }
 
-TUniquePtr<FChange> FSetRichCurveKeyAction::ExecuteInternal(IAnimationDataModel* Model, IAnimationDataController* Controller)
+TUniquePtr<FChange> FSetRichCurveKeyAction::ExecuteInternal(UAnimDataModel* Model, UAnimDataController* Controller)
 {
 	const FRichCurve& RichCurve = Model->GetRichCurve(CurveId);
 
@@ -348,7 +350,7 @@ FString FSetRichCurveKeyAction::ToStringInternal() const
 	return FText::Format(LOCTEXT("SetNamedRichCurveKeyAction_Description", "Setting key for {0} curve '{1}'."), FText::FromString(CurveId.CurveType == ERawCurveTrackTypes::RCT_Float ? FloatLabel : TransformLabel), FText::FromName(CurveId.InternalName.DisplayName)).ToString();
 }
 
-TUniquePtr<FChange> FRemoveRichCurveKeyAction::ExecuteInternal(IAnimationDataModel* Model, IAnimationDataController* Controller)
+TUniquePtr<FChange> FRemoveRichCurveKeyAction::ExecuteInternal(UAnimDataModel* Model, UAnimDataController* Controller)
 {	
 	const FRichCurve& RichCurve = Model->GetRichCurve(CurveId);
 
@@ -369,7 +371,7 @@ FString FRemoveRichCurveKeyAction::ToStringInternal() const
 	return FText::Format(LOCTEXT("RemoveNamedRichCurveKeyAction_Description", "Removing key from {0} curve '{1}'."), FText::FromString(CurveId.CurveType == ERawCurveTrackTypes::RCT_Float ? FloatLabel : TransformLabel), FText::FromName(CurveId.InternalName.DisplayName)).ToString();
 }
 
-TUniquePtr<FChange> FSetRichCurveKeysAction::ExecuteInternal(IAnimationDataModel* Model, IAnimationDataController* Controller)
+TUniquePtr<FChange> FSetRichCurveKeysAction::ExecuteInternal(UAnimDataModel* Model, UAnimDataController* Controller)
 {
 	const FRichCurve& RichCurve = Model->GetRichCurve(CurveId);
 
@@ -388,7 +390,7 @@ FString FSetRichCurveKeysAction::ToStringInternal() const
 	return FText::Format(LOCTEXT("SetNamedRichCurveKeysAction_Description", "Replacing keys for {0} curve '{1}'."), FText::FromString(CurveId.CurveType == ERawCurveTrackTypes::RCT_Float ? FloatLabel : TransformLabel), FText::FromName(CurveId.InternalName.DisplayName)).ToString();
 }
 
-TUniquePtr<FChange> FSetCurveColorAction::ExecuteInternal(IAnimationDataModel* Model, IAnimationDataController* Controller)
+TUniquePtr<FChange> FSetCurveColorAction::ExecuteInternal(UAnimDataModel* Model, UAnimDataController* Controller)
 {
 	const FAnimCurveBase& AnimationCurve = Model->GetCurve(CurveId);
 	const FLinearColor CurrentColor = AnimationCurve.Color;
@@ -408,7 +410,7 @@ FAddAtributeAction::FAddAtributeAction(const FAnimatedBoneAttribute& InAttribute
 	Keys = InAttribute.Curve.GetConstRefOfKeys();
 }
 
-TUniquePtr<FChange> FAddAtributeAction::ExecuteInternal(IAnimationDataModel* Model, IAnimationDataController* Controller)
+TUniquePtr<FChange> FAddAtributeAction::ExecuteInternal(UAnimDataModel* Model, UAnimDataController* Controller)
 {	
 	TArray<const void*> VoidValues;
 	Algo::Transform(Keys, VoidValues, [](const FAttributeKey& Key)
@@ -433,7 +435,7 @@ FString FAddAtributeAction::ToStringInternal() const
 	return FText::Format(LOCTEXT("AddAttributeAction_Description", "Adding attribute '{0}'."), FText::FromName(AttributeId.GetName())).ToString();
 }
 
-TUniquePtr<FChange> FRemoveAtributeAction::ExecuteInternal(IAnimationDataModel* Model, IAnimationDataController* Controller)
+TUniquePtr<FChange> FRemoveAtributeAction::ExecuteInternal(UAnimDataModel* Model, UAnimDataController* Controller)
 {
 	const FAnimatedBoneAttribute& Attribute = Model->GetAttribute(AttributeId);
 	TUniquePtr<FAddAtributeAction> InverseAction = MakeUnique<FAddAtributeAction>(Attribute);
@@ -448,7 +450,7 @@ FString FRemoveAtributeAction::ToStringInternal() const
 	return FText::Format(LOCTEXT("RemoveAttributeAction_Description", "Removing attribute '{0}'."), FText::FromName(AttributeId.GetName())).ToString();
 }
 
-TUniquePtr<FChange> FAddAtributeKeyAction::ExecuteInternal(IAnimationDataModel* Model, IAnimationDataController* Controller)
+TUniquePtr<FChange> FAddAtributeKeyAction::ExecuteInternal(UAnimDataModel* Model, UAnimDataController* Controller)
 {
 	Controller->SetAttributeKey(AttributeId, Key.Time, Key.GetValuePtr<void>(), AttributeId.GetType(), false);
 
@@ -460,7 +462,7 @@ FString FAddAtributeKeyAction::ToStringInternal() const
 	return FText::Format(LOCTEXT("AddAttributeKeyAction_Description", "Adding key to attribute '{0}'."), FText::FromName(AttributeId.GetName())).ToString();
 }
 
-TUniquePtr<FChange> FSetAtributeKeyAction::ExecuteInternal(IAnimationDataModel* Model, IAnimationDataController* Controller)
+TUniquePtr<FChange> FSetAtributeKeyAction::ExecuteInternal(UAnimDataModel* Model, UAnimDataController* Controller)
 {
 	const FAnimatedBoneAttribute& Attribute = Model->GetAttribute(AttributeId);
 
@@ -478,7 +480,7 @@ FString FSetAtributeKeyAction::ToStringInternal() const
 	return FText::Format(LOCTEXT("SetAttributeKeyAction_Description", "Setting key on attribute '{0}'."), FText::FromName(AttributeId.GetName())).ToString();
 }
 
-TUniquePtr<FChange> FRemoveAtributeKeyAction::ExecuteInternal(IAnimationDataModel* Model, IAnimationDataController* Controller)
+TUniquePtr<FChange> FRemoveAtributeKeyAction::ExecuteInternal(UAnimDataModel* Model, UAnimDataController* Controller)
 {
 	const FAnimatedBoneAttribute& Attribute = Model->GetAttribute(AttributeId);
 
@@ -502,7 +504,7 @@ FSetAtributeKeysAction::FSetAtributeKeysAction(const FAnimatedBoneAttribute& InA
 	Keys = InAttribute.Curve.GetConstRefOfKeys();
 }
 
-TUniquePtr<FChange> FSetAtributeKeysAction::ExecuteInternal(IAnimationDataModel* Model, IAnimationDataController* Controller)
+TUniquePtr<FChange> FSetAtributeKeysAction::ExecuteInternal(UAnimDataModel* Model, UAnimDataController* Controller)
 {
 	TUniquePtr<FSetAtributeKeysAction> InverseAction = MakeUnique<FSetAtributeKeysAction>(Model->GetAttribute(AttributeId));
 
diff --git a/Engine/Source/Developer/AnimationDataController/Public/AnimDataController.h b/Engine/Source/Developer/AnimationDataController/Public/AnimDataController.h
--- a/Engine/Source/Developer/AnimationDataController/Public/AnimDataController.h
+++ b/Engine/Source/Developer/AnimationDataController/Public/AnimDataController.h
@@ -15,7 +15,22 @@
 
 #include "AnimDataController.generated.h"
 
-UCLASS(BlueprintType)
+#if WITH_EDITOR
+
+#endif // WITH_EDITOR
+
+struct FAnimationCurveIdentifier;
+struct FAnimationAttributeIdentifier;
+
+namespace UE {
+namespace Anim {
+	class FOpenBracketAction;
+	class FCloseBracketAction;
+	
+	static const int32 DefaultCurveFlags = EAnimAssetCurveFlags::AACF_Editable;
+}}
+
+     UCLASS(BlueprintType)
 class ANIMATIONDATACONTROLLER_API UAnimDataController : public UObject, public IAnimationDataController
 {
 	GENERATED_BODY()
@@ -29,18 +44,15 @@ public:
 
 #if WITH_EDITOR
 	/** Begin IAnimationDataController overrides */
-	virtual void SetModel(TScriptInterface<IAnimationDataModel> InModel) override;
-    virtual TScriptInterface<IAnimationDataModel> GetModelInterface() const override { return ModelInterface; }
-	virtual const IAnimationDataModel* const GetModel() const override { return Model; }
+	virtual void SetModel(UAnimDataModel* InModel)  override;
+    virtual UAnimDataModel* GetModel() override { return Model; }
+	virtual const UAnimDataModel* const GetModel() const override { return Model; }
 	virtual void OpenBracket(const FText& InTitle, bool bShouldTransact = true) override;
 	virtual void CloseBracket(bool bShouldTransact = true) override;
-	virtual void SetNumberOfFrames(FFrameNumber Length, bool bShouldTransact = true) override;
-	virtual void ResizeNumberOfFrames(FFrameNumber NewLength, FFrameNumber T0, FFrameNumber T1, bool bShouldTransact = true) override;
-	virtual void ResizeInFrames(FFrameNumber Length, FFrameNumber T0, FFrameNumber T1, bool bShouldTransact = true) override;
-	virtual void SetFrameRate(FFrameRate FrameRate, bool bShouldTransact = true) override;
 	virtual void SetPlayLength(float Length, bool bShouldTransact = true) override;
 	virtual void ResizePlayLength(float NewLength, float T0, float T1, bool bShouldTransact = true) override;
 	virtual void Resize(float Length, float T0, float T1, bool bShouldTransact = true) override;
+	virtual void SetFrameRate(FFrameRate FrameRate, bool bShouldTransact = true) override;
 	virtual int32 AddBoneTrack(FName BoneName, bool bShouldTransact = true) override;
 	virtual int32 InsertBoneTrack(FName BoneName, int32 DesiredIndex, bool bShouldTransact = true) override;
 	virtual bool RemoveBoneTrack(FName BoneName, bool bShouldTransact = true) override;
@@ -84,9 +96,6 @@ public:
 	}
     virtual bool RemoveAttributeKey(const FAnimationAttributeIdentifier& AttributeIdentifier, float Time, bool bShouldTransact = true) override;	
 	virtual bool DuplicateAttribute(const FAnimationAttributeIdentifier& AttributeIdentifier, const FAnimationAttributeIdentifier& NewAttributeIdentifier, bool bShouldTransact = true) override;
-	virtual void UpdateWithSkeleton(USkeleton* TargetSkeleton, bool bShouldTransact = true) override;
-	virtual void PopulateWithExistingModel(TScriptInterface<IAnimationDataModel> InModel) override;
-	virtual void InitializeModel() override {}
 protected:
 	virtual void NotifyBracketOpen() override;
 	virtual void NotifyBracketClosed() override;
@@ -97,22 +106,48 @@ private:
 	bool SetAttributeKey_Internal(const FAnimationAttributeIdentifier& AttributeIdentifier, float Time, const void* KeyValue, const UScriptStruct* TypeStruct, bool bShouldTransact = true);
 	bool SetAttributeKeys_Internal(const FAnimationAttributeIdentifier& AttributeIdentifier, TArrayView<const float> Times, TArrayView<const void*> KeyValues, const UScriptStruct* TypeStruct, bool bShouldTransact = true);
 
+	/** Returns whether or not the supplied curve type is supported by the controller functionality */
+	const bool IsSupportedCurveType(ERawCurveTrackTypes CurveType) const;
+	/** Returns the string representation of the provided curve enum type value */
+	FString GetCurveTypeValueName(ERawCurveTrackTypes InType) const;
+	
 	/** Resizes the curve/attribute data stored on the model according to the provided new length and time at which to insert or remove time */
 	void ResizeCurves(float NewLength, bool bInserted, float T0, float T1, bool bShouldTransact = true);
 	void ResizeAttributes(float NewLength, bool bInserted, float T0, float T1, bool bShouldTransact = true);
+
+	/** Ensures that a valid model is currently targeted */
+	void ValidateModel() const;
+
+	/** Verifies whether or not the Model's outer object is (or is derived from) the specified UClass */
+	bool CheckOuterClass(UClass* InClass) const;
+
+	/** Helper functionality to output script-based warnings and errors */
+	void ReportWarning(const FText& InMessage) const;
+	void ReportError(const FText& InMessage) const;
+
+	template <typename FmtType, typename... Types>
+    void ReportWarningf(const FmtType& Fmt, Types... Args) const
+	{	
+		ReportWarning(FText::Format(Fmt, Args...));
+	}
+
+	template <typename FmtType, typename... Types>
+    void ReportErrorf(const FmtType& Fmt, Types... Args) const
+	{
+		ReportError(FText::Format(Fmt, Args...));
+	}
 #endif // WITH_EDITOR
 
 private: 
 #if WITH_EDITOR
 	int32 BracketDepth;
+
+	UE::FChangeTransactor ChangeTransactor;	
 #endif // WITH_EDITOR
 
 #if WITH_EDITORONLY_DATA
 	UPROPERTY(transient)
 	TObjectPtr<UAnimDataModel> Model;
-	
-	UPROPERTY(transient)
-	TScriptInterface<IAnimationDataModel> ModelInterface;
 #endif // WITH_EDITORONLY_DATA
 
 	friend class FAnimDataControllerTestBase;
diff --git a/Engine/Source/Developer/AnimationDataController/Public/AnimDataControllerActions.h b/Engine/Source/Developer/AnimationDataController/Public/AnimDataControllerActions.h
--- a/Engine/Source/Developer/AnimationDataController/Public/AnimDataControllerActions.h
+++ b/Engine/Source/Developer/AnimationDataController/Public/AnimDataControllerActions.h
@@ -3,15 +3,14 @@
 #pragma once
 
 #include "Misc/Change.h"
-#include "Misc/CoreMiscDefines.h"
 
-#include "Animation/AnimData/IAnimationDataModel.h"
+#include "Animation/AnimData/AnimDataModel.h"
 #include "Animation/AnimTypes.h" 
 #include "Misc/FrameRate.h"
 
 class UObject;
-class IAnimationDataModel;
-class IAnimationDataController;
+class UAnimDataModel;
+class UAnimDataController;
 
 #if WITH_EDITOR
 
@@ -20,11 +19,11 @@ namespace UE {
 namespace Anim {
 
 /**
-* UAnimDataController instanced FChange-based objects used for storing mutations to an IAnimationDataModel within the Transaction Buffer.
-* Each Action class represents an (invertable) operation mutating an IAnimationDataModel object utilizing a UAnimDataController. Allowing 
+* UAnimDataController instanced FChange-based objects used for storing mutations to an UAnimDataModel within the Transaction Buffer.
+* Each Action class represents an (invertable) operation mutating an UAnimDataModel object utilizing a UAnimDataController. Allowing 
 * for a more granular approach to undo/redo-ing changes while also allowing for script-based interoperability.
 */
-class ANIMATIONDATACONTROLLER_API FAnimDataBaseAction : public FSwapChange
+class FAnimDataBaseAction : public FSwapChange
 {
 public:
 	virtual TUniquePtr<FChange> Execute(UObject* Object) final;
@@ -32,76 +31,78 @@ public:
 	virtual FString ToString() const override;
 
 protected:
-	virtual TUniquePtr<FChange> ExecuteInternal(IAnimationDataModel* Model, IAnimationDataController* Controller) = 0;
+	virtual TUniquePtr<FChange> ExecuteInternal(UAnimDataModel* Model, UAnimDataController* Controller) = 0;
 	virtual FString ToStringInternal() const = 0;
 };
 
-class ANIMATIONDATACONTROLLER_API FOpenBracketAction : public FAnimDataBaseAction
+class FOpenBracketAction : public FAnimDataBaseAction
 {
 public:
 	explicit FOpenBracketAction(const FString& InDescription) : Description(InDescription) {}
 	virtual ~FOpenBracketAction() {}
 protected:
 	FOpenBracketAction() {}
-	virtual TUniquePtr<FChange> ExecuteInternal(IAnimationDataModel* Model, IAnimationDataController* Controller) override;
+	virtual TUniquePtr<FChange> ExecuteInternal(UAnimDataModel* Model, UAnimDataController* Controller) override;
 	virtual FString ToStringInternal() const override;
 
 protected:
 	FString Description;
 };
 
-class ANIMATIONDATACONTROLLER_API FCloseBracketAction : public FAnimDataBaseAction
+class FCloseBracketAction : public FAnimDataBaseAction
 {
 public:
 	explicit FCloseBracketAction(const FString& InDescription) : Description(InDescription) {}
 	virtual ~FCloseBracketAction() {}
 protected:
 	FCloseBracketAction() {}
-	virtual TUniquePtr<FChange> ExecuteInternal(IAnimationDataModel* Model, IAnimationDataController* Controller) override;
+	virtual TUniquePtr<FChange> ExecuteInternal(UAnimDataModel* Model, UAnimDataController* Controller) override;
 	virtual FString ToStringInternal() const override;
 
 protected:
 	FString Description;
 };
 
-class ANIMATIONDATACONTROLLER_API FAddTrackAction : public FAnimDataBaseAction
+class FAddTrackAction : public FAnimDataBaseAction
 {
 public:
-	explicit FAddTrackAction(const FBoneAnimationTrack& Track);
+	explicit FAddTrackAction(const FBoneAnimationTrack& Track, int32 TrackIndex);
 	virtual ~FAddTrackAction() {}
 protected:
 	FAddTrackAction() {}
-	virtual TUniquePtr<FChange> ExecuteInternal(IAnimationDataModel* Model, IAnimationDataController* Controller) override;
+	virtual TUniquePtr<FChange> ExecuteInternal(UAnimDataModel* Model, UAnimDataController* Controller) override;
 	virtual FString ToStringInternal() const override;
 
 protected:
 	FName Name;
 	int32 BoneTreeIndex;
 	FRawAnimSequenceTrack Data;
+	int32 TrackIndex;
 };
 
-class ANIMATIONDATACONTROLLER_API FRemoveTrackAction : public FAnimDataBaseAction
+class FRemoveTrackAction : public FAnimDataBaseAction
 {
 public:
-	explicit FRemoveTrackAction(const FBoneAnimationTrack& Track);
+	explicit FRemoveTrackAction(const FBoneAnimationTrack& Track, int32 TrackIndex);
 	virtual ~FRemoveTrackAction() {}
 protected:
 	FRemoveTrackAction() {}
-	virtual TUniquePtr<FChange> ExecuteInternal(IAnimationDataModel* Model, IAnimationDataController* Controller) override;
+	virtual TUniquePtr<FChange> ExecuteInternal(UAnimDataModel* Model, UAnimDataController* Controller) override;
 	virtual FString ToStringInternal() const override;
 
 protected:
 	FName Name;
+	int32 TrackIndex;
 };
 
-class ANIMATIONDATACONTROLLER_API FSetTrackKeysAction : public FAnimDataBaseAction
+class FSetTrackKeysAction : public FAnimDataBaseAction
 {
 public:
 	explicit FSetTrackKeysAction(const FBoneAnimationTrack& Track);
 	virtual ~FSetTrackKeysAction() {}
 protected:
 	FSetTrackKeysAction() {}
-	virtual TUniquePtr<FChange> ExecuteInternal(IAnimationDataModel* Model, IAnimationDataController* Controller) override;
+	virtual TUniquePtr<FChange> ExecuteInternal(UAnimDataModel* Model, UAnimDataController* Controller) override;
 	virtual FString ToStringInternal() const override;
 
 protected:
@@ -111,37 +112,37 @@ protected:
 	FRawAnimSequenceTrack TrackData;
 };
 
-class ANIMATIONDATACONTROLLER_API FResizePlayLengthInFramesAction : public FAnimDataBaseAction
+class FResizePlayLengthAction : public FAnimDataBaseAction
 {
 public:
-	explicit FResizePlayLengthInFramesAction(const IAnimationDataModel* InModel, FFrameNumber F0, FFrameNumber F1);
-	virtual ~FResizePlayLengthInFramesAction() override {}
+	explicit FResizePlayLengthAction(const UAnimDataModel* InModel, float t0, float t1);
+	virtual ~FResizePlayLengthAction() {}
 protected:
-	FResizePlayLengthInFramesAction() {}
-	virtual TUniquePtr<FChange> ExecuteInternal(IAnimationDataModel* Model, IAnimationDataController* Controller) override;
+	FResizePlayLengthAction() {}
+	virtual TUniquePtr<FChange> ExecuteInternal(UAnimDataModel* Model, UAnimDataController* Controller) override;
 	virtual FString ToStringInternal() const override;
 
 protected:
-	FFrameNumber Length;
-	FFrameNumber Frame0;
-	FFrameNumber Frame1;
+	float Length;
+	float T0;
+	float T1;
 };
 
-class ANIMATIONDATACONTROLLER_API FSetFrameRateAction : public FAnimDataBaseAction
+class FSetFrameRateAction : public FAnimDataBaseAction
 {
 public:
-	explicit FSetFrameRateAction(const IAnimationDataModel* InModel);
-	virtual ~FSetFrameRateAction() override {}
+	explicit FSetFrameRateAction(const UAnimDataModel* InModel);
+	virtual ~FSetFrameRateAction() {}
 protected:
 	FSetFrameRateAction() {}
-	virtual TUniquePtr<FChange> ExecuteInternal(IAnimationDataModel* Model, IAnimationDataController* Controller) override;
+	virtual TUniquePtr<FChange> ExecuteInternal(UAnimDataModel* Model, UAnimDataController* Controller) override;
 	virtual FString ToStringInternal() const override;
 
 protected:
 	FFrameRate FrameRate;
 };
 
-class ANIMATIONDATACONTROLLER_API FAddCurveAction : public FAnimDataBaseAction
+class FAddCurveAction : public FAnimDataBaseAction
 {
 public:
 	explicit FAddCurveAction(const FAnimationCurveIdentifier& InCurveId, int32 InFlags) : CurveId(InCurveId), Flags(InFlags) {}
@@ -149,7 +150,7 @@ public:
 	virtual ~FAddCurveAction() {}
 protected:
 	FAddCurveAction() {}
-	virtual TUniquePtr<FChange> ExecuteInternal(IAnimationDataModel* Model, IAnimationDataController* Controller) override;
+	virtual TUniquePtr<FChange> ExecuteInternal(UAnimDataModel* Model, UAnimDataController* Controller) override;
 	virtual FString ToStringInternal() const override;
 
 protected:
@@ -157,14 +158,14 @@ protected:
 	int32 Flags;
 };
 
-class ANIMATIONDATACONTROLLER_API FAddFloatCurveAction : public FAnimDataBaseAction
+class FAddFloatCurveAction : public FAnimDataBaseAction
 {
 public:
 	explicit FAddFloatCurveAction(const FAnimationCurveIdentifier& InCurveId, int32 InFlags, const TArray<FRichCurveKey>& InKeys, const FLinearColor& InColor) : CurveId(InCurveId), Flags(InFlags), Keys(InKeys), Color(InColor) {}
 	virtual ~FAddFloatCurveAction() {}
 protected:
 	FAddFloatCurveAction() {}
-	virtual TUniquePtr<FChange> ExecuteInternal(IAnimationDataModel* Model, IAnimationDataController* Controller) override;
+	virtual TUniquePtr<FChange> ExecuteInternal(UAnimDataModel* Model, UAnimDataController* Controller) override;
 	virtual FString ToStringInternal() const override;
 
 protected:
@@ -174,14 +175,14 @@ protected:
 	FLinearColor Color;
 };
 
-class ANIMATIONDATACONTROLLER_API FAddTransformCurveAction : public FAnimDataBaseAction
+class FAddTransformCurveAction : public FAnimDataBaseAction
 {
 public:
 	explicit FAddTransformCurveAction(const FAnimationCurveIdentifier& InCurveId, int32 InFlags, const FTransformCurve& InTransformCurve);
 	virtual ~FAddTransformCurveAction() {}
 protected:
 	FAddTransformCurveAction() {}
-	virtual TUniquePtr<FChange> ExecuteInternal(IAnimationDataModel* Model, IAnimationDataController* Controller) override;
+	virtual TUniquePtr<FChange> ExecuteInternal(UAnimDataModel* Model, UAnimDataController* Controller) override;
 	virtual FString ToStringInternal() const override;
 
 protected:
@@ -191,28 +192,28 @@ protected:
 	TArray<FRichCurveKey> SubCurveKeys[9];
 };
 
-class ANIMATIONDATACONTROLLER_API FRemoveCurveAction : public FAnimDataBaseAction
+class FRemoveCurveAction : public FAnimDataBaseAction
 {
 public:
 	explicit FRemoveCurveAction(const FAnimationCurveIdentifier& InCurveId) : CurveId(InCurveId) {}
 	virtual ~FRemoveCurveAction() {}
 protected:
 	FRemoveCurveAction() {}
-	virtual TUniquePtr<FChange> ExecuteInternal(IAnimationDataModel* Model, IAnimationDataController* Controller) override;
+	virtual TUniquePtr<FChange> ExecuteInternal(UAnimDataModel* Model, UAnimDataController* Controller) override;
 	virtual FString ToStringInternal() const override;
 
 protected:
 	FAnimationCurveIdentifier CurveId;
 };
 
-class ANIMATIONDATACONTROLLER_API FSetCurveFlagsAction : public FAnimDataBaseAction
+class FSetCurveFlagsAction : public FAnimDataBaseAction
 {
 public:
 	explicit FSetCurveFlagsAction(const FAnimationCurveIdentifier& InCurveId, int32 InFlags, ERawCurveTrackTypes InCurveType) : CurveId(InCurveId), Flags(InFlags), CurveType(InCurveType) {}
 	virtual ~FSetCurveFlagsAction() {}
 protected:
 	FSetCurveFlagsAction() {}
-	virtual TUniquePtr<FChange> ExecuteInternal(IAnimationDataModel* Model, IAnimationDataController* Controller) override;
+	virtual TUniquePtr<FChange> ExecuteInternal(UAnimDataModel* Model, UAnimDataController* Controller) override;
 	virtual FString ToStringInternal() const override;
 
 protected:
@@ -221,14 +222,14 @@ protected:
 	ERawCurveTrackTypes CurveType;
 };
 
-class ANIMATIONDATACONTROLLER_API FRenameCurveAction : public FAnimDataBaseAction
+class FRenameCurveAction : public FAnimDataBaseAction
 {
 public:
 	explicit FRenameCurveAction(const FAnimationCurveIdentifier& InCurveId, const FAnimationCurveIdentifier& InNewCurveId) : CurveId(InCurveId), NewCurveId(InNewCurveId) {}
 	virtual ~FRenameCurveAction() {}
 protected:
 	FRenameCurveAction() {}
-	virtual TUniquePtr<FChange> ExecuteInternal(IAnimationDataModel* Model, IAnimationDataController* Controller) override;
+	virtual TUniquePtr<FChange> ExecuteInternal(UAnimDataModel* Model, UAnimDataController* Controller) override;
 	virtual FString ToStringInternal() const override;
 
 protected:
@@ -236,14 +237,14 @@ protected:
 	FAnimationCurveIdentifier NewCurveId;
 };
 
-class ANIMATIONDATACONTROLLER_API FScaleCurveAction : public FAnimDataBaseAction
+class FScaleCurveAction : public FAnimDataBaseAction
 {
 public:
 	explicit FScaleCurveAction(const FAnimationCurveIdentifier& InCurveId, float InOrigin, float InFactor, ERawCurveTrackTypes InCurveType);
 	virtual ~FScaleCurveAction() {}
 protected:
 	FScaleCurveAction() {}
-	virtual TUniquePtr<FChange> ExecuteInternal(IAnimationDataModel* Model, IAnimationDataController* Controller) override;
+	virtual TUniquePtr<FChange> ExecuteInternal(UAnimDataModel* Model, UAnimDataController* Controller) override;
 	virtual FString ToStringInternal() const override;
 
 protected:
@@ -253,14 +254,14 @@ protected:
 	float Factor;
 };
 
-class ANIMATIONDATACONTROLLER_API FAddRichCurveKeyAction : public FAnimDataBaseAction
+class FAddRichCurveKeyAction : public FAnimDataBaseAction
 {
 public:
 	explicit FAddRichCurveKeyAction(const FAnimationCurveIdentifier& InCurveId, const FRichCurveKey& InKey) : CurveId(InCurveId), Key(InKey) {}
 	virtual ~FAddRichCurveKeyAction() {}
 protected:
 	FAddRichCurveKeyAction() {}
-	virtual TUniquePtr<FChange> ExecuteInternal(IAnimationDataModel* Model, IAnimationDataController* Controller) override;
+	virtual TUniquePtr<FChange> ExecuteInternal(UAnimDataModel* Model, UAnimDataController* Controller) override;
 	virtual FString ToStringInternal() const override;
 
 protected:
@@ -269,14 +270,14 @@ protected:
 	FRichCurveKey Key;
 };
 
-class ANIMATIONDATACONTROLLER_API FSetRichCurveKeyAction : public FAnimDataBaseAction
+class FSetRichCurveKeyAction : public FAnimDataBaseAction
 {
 public:
 	explicit FSetRichCurveKeyAction(const FAnimationCurveIdentifier& InCurveId, const FRichCurveKey& InKey) : CurveId(InCurveId), Key(InKey) {}
 	virtual ~FSetRichCurveKeyAction() {}
 protected:
 	FSetRichCurveKeyAction() {}
-	virtual TUniquePtr<FChange> ExecuteInternal(IAnimationDataModel* Model, IAnimationDataController* Controller) override;
+	virtual TUniquePtr<FChange> ExecuteInternal(UAnimDataModel* Model, UAnimDataController* Controller) override;
 	virtual FString ToStringInternal() const override;
 
 protected:
@@ -284,14 +285,14 @@ protected:
 	FRichCurveKey Key;
 };
 
-class ANIMATIONDATACONTROLLER_API FRemoveRichCurveKeyAction : public FAnimDataBaseAction
+class FRemoveRichCurveKeyAction : public FAnimDataBaseAction
 {
 public:
 	explicit FRemoveRichCurveKeyAction(const FAnimationCurveIdentifier& InCurveId, const float InTime) : CurveId(InCurveId), Time(InTime) {}
 	virtual ~FRemoveRichCurveKeyAction() {}
 protected:
 	FRemoveRichCurveKeyAction() {}
-	virtual TUniquePtr<FChange> ExecuteInternal(IAnimationDataModel* Model, IAnimationDataController* Controller) override;
+	virtual TUniquePtr<FChange> ExecuteInternal(UAnimDataModel* Model, UAnimDataController* Controller) override;
 	virtual FString ToStringInternal() const override;
 
 protected:
@@ -299,14 +300,14 @@ protected:
 	float Time;
 };
 
-class ANIMATIONDATACONTROLLER_API FSetRichCurveKeysAction : public FAnimDataBaseAction
+class FSetRichCurveKeysAction : public FAnimDataBaseAction
 {
 public:
 	explicit FSetRichCurveKeysAction(const FAnimationCurveIdentifier& InCurveId, const TArray<FRichCurveKey>& InKeys) : CurveId(InCurveId), Keys(InKeys) {}
 	virtual ~FSetRichCurveKeysAction() {}
 protected:
 	FSetRichCurveKeysAction() {}
-	virtual TUniquePtr<FChange> ExecuteInternal(IAnimationDataModel* Model, IAnimationDataController* Controller) override;
+	virtual TUniquePtr<FChange> ExecuteInternal(UAnimDataModel* Model, UAnimDataController* Controller) override;
 	virtual FString ToStringInternal() const override;
 
 protected:
@@ -314,14 +315,14 @@ protected:
 	TArray<FRichCurveKey> Keys;
 };
 
-class ANIMATIONDATACONTROLLER_API FSetCurveColorAction : public FAnimDataBaseAction
+class FSetCurveColorAction : public FAnimDataBaseAction
 {
 public:
 	explicit FSetCurveColorAction(const FAnimationCurveIdentifier& InCurveId, const FLinearColor& InColor) : CurveId(InCurveId), Color(InColor) {}
 	virtual ~FSetCurveColorAction() {}
 protected:
 	FSetCurveColorAction() {}
-	virtual TUniquePtr<FChange> ExecuteInternal(IAnimationDataModel* Model, IAnimationDataController* Controller) override;
+	virtual TUniquePtr<FChange> ExecuteInternal(UAnimDataModel* Model, UAnimDataController* Controller) override;
 	virtual FString ToStringInternal() const override;
 
 protected:
@@ -329,14 +330,14 @@ protected:
 	FLinearColor Color;
 };
 
-class ANIMATIONDATACONTROLLER_API FAddAtributeAction : public FAnimDataBaseAction
+class FAddAtributeAction : public FAnimDataBaseAction
 {
 public:
 	explicit FAddAtributeAction(const FAnimatedBoneAttribute& InAttribute);
 	virtual ~FAddAtributeAction() {}
 protected:
 	FAddAtributeAction() {}
-	virtual TUniquePtr<FChange> ExecuteInternal(IAnimationDataModel* Model, IAnimationDataController* Controller) override;
+	virtual TUniquePtr<FChange> ExecuteInternal(UAnimDataModel* Model, UAnimDataController* Controller) override;
 	virtual FString ToStringInternal() const override;
 
 protected:
@@ -344,28 +345,28 @@ protected:
 	TArray<FAttributeKey> Keys;
 };
 
-class ANIMATIONDATACONTROLLER_API FRemoveAtributeAction : public FAnimDataBaseAction
+class FRemoveAtributeAction : public FAnimDataBaseAction
 {
 public:
 	explicit FRemoveAtributeAction(const FAnimationAttributeIdentifier& InAttributeId) : AttributeId(InAttributeId) {}
 	virtual ~FRemoveAtributeAction() {}
 protected:
 	FRemoveAtributeAction() {}
-	virtual TUniquePtr<FChange> ExecuteInternal(IAnimationDataModel* Model, IAnimationDataController* Controller) override;
+	virtual TUniquePtr<FChange> ExecuteInternal(UAnimDataModel* Model, UAnimDataController* Controller) override;
 	virtual FString ToStringInternal() const override;
 
 protected:
 	FAnimationAttributeIdentifier AttributeId;
 };
 
-class ANIMATIONDATACONTROLLER_API FAddAtributeKeyAction : public FAnimDataBaseAction
+class FAddAtributeKeyAction : public FAnimDataBaseAction
 {
 public:
 	explicit FAddAtributeKeyAction(const FAnimationAttributeIdentifier& InAttributeId, const FAttributeKey& InKey) : AttributeId(InAttributeId), Key(InKey) {}
 	virtual ~FAddAtributeKeyAction() {}
 protected:
 	FAddAtributeKeyAction() {}
-	virtual TUniquePtr<FChange> ExecuteInternal(IAnimationDataModel* Model, IAnimationDataController* Controller) override;
+	virtual TUniquePtr<FChange> ExecuteInternal(UAnimDataModel* Model, UAnimDataController* Controller) override;
 	virtual FString ToStringInternal() const override;
 
 protected:
@@ -373,14 +374,14 @@ protected:
 	FAttributeKey Key;
 };
 
-class ANIMATIONDATACONTROLLER_API FSetAtributeKeyAction : public FAnimDataBaseAction
+class FSetAtributeKeyAction : public FAnimDataBaseAction
 {
 public:
 	explicit FSetAtributeKeyAction(const FAnimationAttributeIdentifier& InAttributeId, const FAttributeKey& InKey) : AttributeId(InAttributeId), Key(InKey) {}
 	virtual ~FSetAtributeKeyAction() {}
 protected:
 	FSetAtributeKeyAction() {}
-	virtual TUniquePtr<FChange> ExecuteInternal(IAnimationDataModel* Model, IAnimationDataController* Controller) override;
+	virtual TUniquePtr<FChange> ExecuteInternal(UAnimDataModel* Model, UAnimDataController* Controller) override;
 	virtual FString ToStringInternal() const override;
 
 protected:
@@ -388,14 +389,14 @@ protected:
 	FAttributeKey Key;
 };
 
-class ANIMATIONDATACONTROLLER_API FRemoveAtributeKeyAction : public FAnimDataBaseAction
+class FRemoveAtributeKeyAction : public FAnimDataBaseAction
 {
 public:
 	explicit FRemoveAtributeKeyAction(const FAnimationAttributeIdentifier& InAttributeId, float InTime) : AttributeId(InAttributeId), Time(InTime) {}
 	virtual ~FRemoveAtributeKeyAction() {}
 protected:
 	FRemoveAtributeKeyAction() {}
-	virtual TUniquePtr<FChange> ExecuteInternal(IAnimationDataModel* Model, IAnimationDataController* Controller) override;
+	virtual TUniquePtr<FChange> ExecuteInternal(UAnimDataModel* Model, UAnimDataController* Controller) override;
 	virtual FString ToStringInternal() const override;
 
 protected:
@@ -403,14 +404,14 @@ protected:
 	float Time;
 };
 
-class ANIMATIONDATACONTROLLER_API FSetAtributeKeysAction : public FAnimDataBaseAction
+class FSetAtributeKeysAction : public FAnimDataBaseAction
 {
 public:
 	explicit FSetAtributeKeysAction(const FAnimatedBoneAttribute& InAttribute);
 	virtual ~FSetAtributeKeysAction() {}
 protected:
 	FSetAtributeKeysAction() {}
-	virtual TUniquePtr<FChange> ExecuteInternal(IAnimationDataModel* Model, IAnimationDataController* Controller) override;
+	virtual TUniquePtr<FChange> ExecuteInternal(UAnimDataModel* Model, UAnimDataController* Controller) override;
 	virtual FString ToStringInternal() const override;
 
 protected:
diff --git a/Engine/Source/Developer/MeshBoneReduction/Private/MeshBoneReduction.cpp b/Engine/Source/Developer/MeshBoneReduction/Private/MeshBoneReduction.cpp
--- a/Engine/Source/Developer/MeshBoneReduction/Private/MeshBoneReduction.cpp
+++ b/Engine/Source/Developer/MeshBoneReduction/Private/MeshBoneReduction.cpp
@@ -51,7 +51,7 @@ public:
 		}
 	}
 
-	bool GetBoneReductionData(const USkeletalMesh* SkeletalMesh, int32 DesiredLOD, TMap<FBoneIndexType, FBoneIndexType>& OutBonesToReplace, const TArray<FName>* BoneNamesToRemove = nullptr) override
+	bool GetBoneReductionData(const USkeletalMesh* SkeletalMesh, int32 DesiredLOD, TMap<FBoneIndexType, FBoneIndexType>& OutBonesToReplace, const TArray<FName>* BoneNamesToRemove = NULL) override
 	{
 		if (!SkeletalMesh)
 		{
@@ -320,20 +320,14 @@ public:
 			Pose.SetBoneContainer(&RequiredBones);
 			Pose.ResetToRefPose();
 
-			FBlendedCurve TempCurve;
-			TempCurve.InitFrom(RequiredBones);
-			UE::Anim::FStackAttributeContainer TempAttributes;
-
-			FAnimationPoseData AnimPoseData(Pose, TempCurve, TempAttributes);
-
 			// Get component space retarget base pose, will be equivalent of ref-pose if not edited
 			TArray<FTransform> ComponentSpaceRefPose;
-			FAnimationRuntime::FillUpComponentSpaceTransforms(SkeletalMesh->GetRefSkeleton(), SkeletalMesh->GetRefSkeleton().GetRefBonePose(), ComponentSpaceRefPose);
+			FAnimationRuntime::FillUpComponentSpaceTransformsRetargetBasePose(SkeletalMesh, ComponentSpaceRefPose);
 			
 			// Retrieve animated pose from anim sequence (including retargeting)
 			const USkeleton* Skeleton = SkeletalMesh->GetSkeleton();
 			const FName RetargetSource = Skeleton->GetRetargetSourceForMesh(SkeletalMesh);
-			UE::Anim::BuildPoseFromModel(BakePoseAnim->GetDataModel(), AnimPoseData, 0.0, EAnimInterpolationType::Step, RetargetSource, Skeleton->GetRefLocalPoses(RetargetSource));
+			UE::Anim::BuildPoseFromModel(BakePoseAnim->GetDataModel(), Pose, 0.f, EAnimInterpolationType::Step, RetargetSource, Skeleton->GetRefLocalPoses(RetargetSource));
 			
 			// Calculate component space animated pose matrices
 			TArray<FMatrix> ComponentSpaceAnimatedPose;
diff --git a/Engine/Source/Editor/AnimGraph/Private/AnimPreviewInstance.cpp b/Engine/Source/Editor/AnimGraph/Private/AnimPreviewInstance.cpp
--- a/Engine/Source/Editor/AnimGraph/Private/AnimPreviewInstance.cpp
+++ b/Engine/Source/Editor/AnimGraph/Private/AnimPreviewInstance.cpp
@@ -4,7 +4,6 @@
 #include "AnimPreviewInstance.h"
 #include "Animation/DebugSkelMeshComponent.h"
 #include "AnimationRuntime.h"
-#include "Animation/AnimationSettings.h"
 #include "Animation/AnimSequence.h"
 #include "Animation/AnimSequenceHelpers.h"
 
@@ -794,12 +793,7 @@ void UAnimPreviewInstance::MontagePreview_StepBackward()
 
 float UAnimPreviewInstance::MontagePreview_CalculateStepLength()
 {
-	if (UAnimMontage* Montage = Cast<UAnimMontage>(CurrentAsset))
-	{
-		return Montage->GetSamplingFrameRate().AsInterval();
-	}
-
-	return UAnimationSettings::Get()->GetDefaultFrameRate().AsInterval();
+	return 1.0f / 30.0f;
 }
 
 void UAnimPreviewInstance::MontagePreview_JumpToStart()
diff --git a/Engine/Source/Editor/AnimationBlueprintLibrary/Private/AnimPose.cpp b/Engine/Source/Editor/AnimationBlueprintLibrary/Private/AnimPose.cpp
--- a/Engine/Source/Editor/AnimationBlueprintLibrary/Private/AnimPose.cpp
+++ b/Engine/Source/Editor/AnimationBlueprintLibrary/Private/AnimPose.cpp
@@ -499,11 +499,12 @@ float UAnimPoseExtensions::GetCurveWeight(const FAnimPose& Pose, const FName& Cu
 
 void UAnimPoseExtensions::GetAnimPoseAtFrame(const UAnimSequenceBase* AnimationSequenceBase, int32 FrameIndex, FAnimPoseEvaluationOptions EvaluationOptions, FAnimPose& Pose)
 {
-	const double Time = AnimationSequenceBase->GetDataModel()->GetFrameRate().AsSeconds(FrameIndex);
+	float Time = 0.f;
+	UAnimationBlueprintLibrary::GetTimeAtFrame(AnimationSequenceBase, FrameIndex, Time);
 	GetAnimPoseAtTime(AnimationSequenceBase, Time, EvaluationOptions, Pose);
 }
 
-void UAnimPoseExtensions::GetAnimPoseAtTime(const UAnimSequenceBase* AnimationSequenceBase, double Time, FAnimPoseEvaluationOptions EvaluationOptions, FAnimPose& Pose)
+void UAnimPoseExtensions::GetAnimPoseAtTime(const UAnimSequenceBase* AnimationSequenceBase, float Time, FAnimPoseEvaluationOptions EvaluationOptions, FAnimPose& Pose)
 {
 	TArray<FAnimPose> InOutPoses;
 	GetAnimPoseAtTimeIntervals(AnimationSequenceBase, { Time }, EvaluationOptions, InOutPoses);
@@ -515,7 +516,7 @@ void UAnimPoseExtensions::GetAnimPoseAtTime(const UAnimSequenceBase* AnimationSe
 	}
 } 
 
-void UAnimPoseExtensions::GetAnimPoseAtTimeIntervals(const UAnimSequenceBase* AnimationSequenceBase, TArray<double> TimeIntervals, FAnimPoseEvaluationOptions EvaluationOptions, TArray<FAnimPose>& InOutPoses)
+void UAnimPoseExtensions::GetAnimPoseAtTimeIntervals(const UAnimSequenceBase* AnimationSequenceBase, TArray<float> TimeIntervals, FAnimPoseEvaluationOptions EvaluationOptions, TArray<FAnimPose>& InOutPoses)
 {
 	if (AnimationSequenceBase && AnimationSequenceBase->GetSkeleton())
 	{
@@ -555,7 +556,7 @@ void UAnimPoseExtensions::GetAnimPoseAtTimeIntervals(const UAnimSequenceBase* An
         UE::Anim::FStackAttributeContainer Attributes;
 
         FAnimationPoseData PoseData(CompactPose, Curve, Attributes);
-        FAnimExtractContext Context(0.0, EvaluationOptions.bExtractRootMotion);
+        FAnimExtractContext Context(0.f, EvaluationOptions.bExtractRootMotion);
     
         FCompactPose BasePose;
         BasePose.SetBoneContainer(&RequiredBones);
@@ -568,7 +569,7 @@ void UAnimPoseExtensions::GetAnimPoseAtTimeIntervals(const UAnimSequenceBase* An
 		
 		for (int32 Index = 0; Index < TimeIntervals.Num(); ++Index)
 		{
-			const double EvalInterval = TimeIntervals[Index];
+			const float EvalInterval = TimeIntervals[Index];
 			
 			bool bValidTime = false;
 			UAnimationBlueprintLibrary::IsValidTime(AnimationSequenceBase, EvalInterval, bValidTime);
diff --git a/Engine/Source/Editor/AnimationBlueprintLibrary/Private/AnimationAttributeBlueprintLibrary.cpp b/Engine/Source/Editor/AnimationBlueprintLibrary/Private/AnimationAttributeBlueprintLibrary.cpp
--- a/Engine/Source/Editor/AnimationBlueprintLibrary/Private/AnimationAttributeBlueprintLibrary.cpp
+++ b/Engine/Source/Editor/AnimationBlueprintLibrary/Private/AnimationAttributeBlueprintLibrary.cpp
@@ -131,14 +131,14 @@ DEFINE_FUNCTION(UAnimationAttributeBlueprintLibrary::execSetAttributeKeys)
 }
 
 
-bool UAnimationAttributeBlueprintLibrary::GetAttributeKey(TScriptInterface<IAnimationDataModel> AnimationDataModel, const FAnimationAttributeIdentifier& AttributeIdentifier, float Time, int32& Value)
+bool UAnimationAttributeBlueprintLibrary::GetAttributeKey(UAnimDataModel* AnimationDataModel, const FAnimationAttributeIdentifier& AttributeIdentifier, float Time, int32& Value)
 {
 	// We should never hit this! Stubbed to avoid NoExport on the class.
 	check(0);
 	return false;
 }
 
-bool UAnimationAttributeBlueprintLibrary::Generic_GetAttributeKey(TScriptInterface<IAnimationDataModel> AnimationDataModel, const FAnimationAttributeIdentifier& AttributeIdentifier, float Time, UScriptStruct* ScriptStruct, void* ValuePtr)
+bool UAnimationAttributeBlueprintLibrary::Generic_GetAttributeKey(UAnimDataModel* AnimationDataModel, const FAnimationAttributeIdentifier& AttributeIdentifier, float Time, UScriptStruct* ScriptStruct, void* ValuePtr)
 {
 	const FAnimatedBoneAttribute* AttributePtr = AnimationDataModel->FindAttribute(AttributeIdentifier);
 
@@ -164,7 +164,7 @@ bool UAnimationAttributeBlueprintLibrary::Generic_GetAttributeKey(TScriptInterfa
 
 DEFINE_FUNCTION(UAnimationAttributeBlueprintLibrary::execGetAttributeKey)
 {
-	P_GET_TINTERFACE(IAnimationDataModel, AnimationDataModel);
+	P_GET_OBJECT(UAnimDataModel, AnimationDataModel);
 	P_GET_STRUCT_REF(FAnimationAttributeIdentifier, AttributeIdentifier);
 	P_GET_PROPERTY(FFloatProperty, TimeInterval);
 
@@ -186,7 +186,7 @@ DEFINE_FUNCTION(UAnimationAttributeBlueprintLibrary::execGetAttributeKey)
 		FBlueprintCoreDelegates::ThrowScriptException(P_THIS, Stack, ExceptionInfo);
 	}
 
-	if (!AnimationDataModel.GetObject())
+	if (!AnimationDataModel)
 	{
 		FBlueprintExceptionInfo ExceptionInfo(
 			EBlueprintExceptionType::AccessViolation, 
@@ -208,14 +208,14 @@ DEFINE_FUNCTION(UAnimationAttributeBlueprintLibrary::execGetAttributeKey)
 }
 
 
-bool UAnimationAttributeBlueprintLibrary::GetAttributeKeys(TScriptInterface<IAnimationDataModel> AnimationDataModel, const FAnimationAttributeIdentifier& AttributeIdentifier, TArray<float>& OutTimes, TArray<int32>& OutValues)
+bool UAnimationAttributeBlueprintLibrary::GetAttributeKeys(UAnimDataModel* AnimationDataModel, const FAnimationAttributeIdentifier& AttributeIdentifier, TArray<float>& OutTimes, TArray<int32>& OutValues)
 {
 	// We should never hit this! Stubbed to avoid NoExport on the class.
 	check(0);
 	return false;
 }
 
-bool UAnimationAttributeBlueprintLibrary::Generic_GetAttributeKeys(TScriptInterface<IAnimationDataModel> AnimationDataModel, const FAnimationAttributeIdentifier& AttributeIdentifier, TArray<float>& Times, void* ValuesArray, const FArrayProperty* ValuesArrayProperty)
+bool UAnimationAttributeBlueprintLibrary::Generic_GetAttributeKeys(UAnimDataModel* AnimationDataModel, const FAnimationAttributeIdentifier& AttributeIdentifier, TArray<float>& Times, void* ValuesArray, const FArrayProperty* ValuesArrayProperty)
 {
 	const FAnimatedBoneAttribute* AttributePtr = AnimationDataModel->FindAttribute(AttributeIdentifier);
 
@@ -249,7 +249,7 @@ bool UAnimationAttributeBlueprintLibrary::Generic_GetAttributeKeys(TScriptInterf
 
 DEFINE_FUNCTION(UAnimationAttributeBlueprintLibrary::execGetAttributeKeys)
 {
-	P_GET_TINTERFACE(IAnimationDataModel, AnimationDataModel);
+	P_GET_OBJECT(UAnimDataModel, AnimationDataModel);
 	P_GET_STRUCT_REF(FAnimationAttributeIdentifier, AttributeIdentifier);
 	P_GET_TARRAY_REF(float, TimeIntervals);
 
@@ -271,7 +271,7 @@ DEFINE_FUNCTION(UAnimationAttributeBlueprintLibrary::execGetAttributeKeys)
 		FBlueprintCoreDelegates::ThrowScriptException(P_THIS, Stack, ExceptionInfo);
 	}
 	
-	if (!AnimationDataModel.GetObject())
+	if (!AnimationDataModel)
 	{
 		FBlueprintExceptionInfo ExceptionInfo(
 			EBlueprintExceptionType::AccessViolation, 
diff --git a/Engine/Source/Editor/AnimationBlueprintLibrary/Private/AnimationBlueprintLibrary.cpp b/Engine/Source/Editor/AnimationBlueprintLibrary/Private/AnimationBlueprintLibrary.cpp
--- a/Engine/Source/Editor/AnimationBlueprintLibrary/Private/AnimationBlueprintLibrary.cpp
+++ b/Engine/Source/Editor/AnimationBlueprintLibrary/Private/AnimationBlueprintLibrary.cpp
@@ -24,7 +24,7 @@
 #include "Modules/ModuleInterface.h"
 
 #include "AnimationRuntime.h"
-#include "Animation/AnimData/IAnimationDataModel.h"
+#include "Animation/AnimData/AnimDataModel.h"
 #include "Animation/AnimSequenceHelpers.h"
 #include "AnimGraphNode_AssetPlayerBase.h"
 #include "Kismet2/BlueprintEditorUtils.h"
@@ -2197,7 +2197,7 @@ bool UAnimationBlueprintLibrary::EvaluateRootBoneTimecodeAttributesAtTime(const
 		return false;
 	}
 
-	const IAnimationDataModel* AnimDataModel = AnimationSequenceBase->GetDataModel();
+	const UAnimDataModel* AnimDataModel = AnimationSequenceBase->GetDataModel();
 	if (!AnimDataModel)
 	{
 		return false;
@@ -2386,7 +2386,7 @@ bool UAnimationBlueprintLibrary::EvaluateRootBoneTimecodeSubframeAttributeAtTime
 		return false;
 	}
 
-	const IAnimationDataModel* AnimDataModel = AnimationSequenceBase->GetDataModel();
+	const UAnimDataModel* AnimDataModel = AnimationSequenceBase->GetDataModel();
 	if (!AnimDataModel)
 	{
 		return false;
diff --git a/Engine/Source/Editor/AnimationBlueprintLibrary/Public/AnimPose.h b/Engine/Source/Editor/AnimationBlueprintLibrary/Public/AnimPose.h
--- a/Engine/Source/Editor/AnimationBlueprintLibrary/Public/AnimPose.h
+++ b/Engine/Source/Editor/AnimationBlueprintLibrary/Public/AnimPose.h
@@ -229,7 +229,7 @@ public:
 	* @param	Pose					Anim Pose to hold the evaluated data
 	*/
 	UFUNCTION(BlueprintCallable, meta = (ScriptMethod), Category = "Animation|Pose")
-	static void GetAnimPoseAtTime(const UAnimSequenceBase* AnimationSequenceBase, double Time, FAnimPoseEvaluationOptions EvaluationOptions, FAnimPose& Pose);
+	static void GetAnimPoseAtTime(const UAnimSequenceBase* AnimationSequenceBase, float Time, FAnimPoseEvaluationOptions EvaluationOptions, FAnimPose& Pose);
 
 	/**
 	* Evaluates an Animation Sequence Base at different time intervals to generate a valid Anim Pose instances
@@ -239,7 +239,7 @@ public:
 	* @param	EvaluationOptions		Options determining the way the pose should be evaluated
 	* @param	InOutPoses				Anim Poses holding the evaluated data (number matches TimeIntervals)
 	*/
-	static void GetAnimPoseAtTimeIntervals(const UAnimSequenceBase* AnimationSequenceBase, TArray<double> TimeIntervals, FAnimPoseEvaluationOptions EvaluationOptions, TArray<FAnimPose>& InOutPoses);
+	static void GetAnimPoseAtTimeIntervals(const UAnimSequenceBase* AnimationSequenceBase, TArray<float> TimeIntervals, FAnimPoseEvaluationOptions EvaluationOptions, TArray<FAnimPose>& InOutPoses);
 
 	/**
 	* Evaluates an Animation Sequence Base to generate a valid Anim Pose instance
diff --git a/Engine/Source/Editor/AnimationBlueprintLibrary/Public/AnimationAttributeBlueprintLibrary.h b/Engine/Source/Editor/AnimationBlueprintLibrary/Public/AnimationAttributeBlueprintLibrary.h
--- a/Engine/Source/Editor/AnimationBlueprintLibrary/Public/AnimationAttributeBlueprintLibrary.h
+++ b/Engine/Source/Editor/AnimationBlueprintLibrary/Public/AnimationAttributeBlueprintLibrary.h
@@ -26,12 +26,12 @@ public:
 
 
 	UFUNCTION(BlueprintCallable, CustomThunk, Category = "Animation|Attributes", meta=(CustomStructureParam="Value", BlueprintInternalUseOnly="true"))
-	static bool GetAttributeKey(TScriptInterface<IAnimationDataModel> AnimationDataModel, const FAnimationAttributeIdentifier& AttributeIdentifier, float Time, int32& Value);
-	static bool Generic_GetAttributeKey(TScriptInterface<IAnimationDataModel> AnimationDataModel, const FAnimationAttributeIdentifier& AttributeIdentifier, float Time, UScriptStruct* ScriptStruct, void* Value);
+	static bool GetAttributeKey(UAnimDataModel* AnimationDataModel, const FAnimationAttributeIdentifier& AttributeIdentifier, float Time, int32& Value);
+	static bool Generic_GetAttributeKey(UAnimDataModel* AnimationDataModel, const FAnimationAttributeIdentifier& AttributeIdentifier, float Time, UScriptStruct* ScriptStruct, void* Value);
 	DECLARE_FUNCTION(execGetAttributeKey);
 
 	UFUNCTION(BlueprintCallable, CustomThunk, Category = "Animation|Attributes", meta=(CustomStructureParam="Values", BlueprintInternalUseOnly="true"))
-	static bool GetAttributeKeys(TScriptInterface<IAnimationDataModel> AnimationDataModel, const FAnimationAttributeIdentifier& AttributeIdentifier, TArray<float>& OutTimes, TArray<int32>& Values);
-	static bool Generic_GetAttributeKeys(TScriptInterface<IAnimationDataModel> AnimationDataModel, const FAnimationAttributeIdentifier& AttributeIdentifier, TArray<float>& Times, void* ValuesArray, const FArrayProperty* ValuesArrayProperty);
+	static bool GetAttributeKeys(UAnimDataModel* AnimationDataModel, const FAnimationAttributeIdentifier& AttributeIdentifier, TArray<float>& OutTimes, TArray<int32>& Values);
+	static bool Generic_GetAttributeKeys(UAnimDataModel* AnimationDataModel, const FAnimationAttributeIdentifier& AttributeIdentifier, TArray<float>& Times, void* ValuesArray, const FArrayProperty* ValuesArrayProperty);
 	DECLARE_FUNCTION(execGetAttributeKeys);
 };
diff --git a/Engine/Source/Editor/MovieSceneTools/Private/EditModes/SkeletalAnimationTrackEditMode.cpp b/Engine/Source/Editor/MovieSceneTools/Private/EditModes/SkeletalAnimationTrackEditMode.cpp
--- a/Engine/Source/Editor/MovieSceneTools/Private/EditModes/SkeletalAnimationTrackEditMode.cpp
+++ b/Engine/Source/Editor/MovieSceneTools/Private/EditModes/SkeletalAnimationTrackEditMode.cpp
@@ -364,7 +364,7 @@ void FSkeletalAnimationTrackEditMode::Render(const FSceneView* View, FViewport*
 											UE::Anim::FStackAttributeContainer TempAttributes;
 											FAnimationPoseData OutAnimationPoseData(OutPose, OutCurve, TempAttributes);
 
-											const double Seconds = AnimSection->MapTimeToAnimation(CurrentTime.Time, CurrentTime.Rate);
+											const float Seconds = static_cast<float>(AnimSection->MapTimeToAnimation(CurrentTime.Time, CurrentTime.Rate));
 											FAnimExtractContext ExtractionContext(Seconds, false);
 
 											AnimSequence->GetAnimationPose(OutAnimationPoseData, ExtractionContext);
diff --git a/Engine/Source/Editor/MovieSceneTools/Private/TrackEditors/TransformTrackEditor.cpp b/Engine/Source/Editor/MovieSceneTools/Private/TrackEditors/TransformTrackEditor.cpp
--- a/Engine/Source/Editor/MovieSceneTools/Private/TrackEditors/TransformTrackEditor.cpp
+++ b/Engine/Source/Editor/MovieSceneTools/Private/TrackEditors/TransformTrackEditor.cpp
@@ -27,7 +27,7 @@
 #include "TransformPropertySection.h"
 #include "SequencerUtilities.h"
 #include "MovieSceneToolHelpers.h"
-#include "Animation/AnimData/IAnimationDataModel.h"
+#include "Animation/AnimData/AnimDataModel.h"
 
 #include "EntitySystem/Interrogation/MovieSceneInterrogationLinker.h"
 #include "EntitySystem/Interrogation/MovieSceneInterrogatedPropertyInstantiator.h"
diff --git a/Engine/Source/Editor/Persona/Private/AnimTimeline/AnimModel_AnimMontage.cpp b/Engine/Source/Editor/Persona/Private/AnimTimeline/AnimModel_AnimMontage.cpp
--- a/Engine/Source/Editor/Persona/Private/AnimTimeline/AnimModel_AnimMontage.cpp
+++ b/Engine/Source/Editor/Persona/Private/AnimTimeline/AnimModel_AnimMontage.cpp
@@ -225,14 +225,14 @@ void FAnimModel_AnimMontage::RecalculateSequenceLength()
 	{
 		AnimMontage->InvalidateRecursiveAsset();
 
-		const float CurrentCalculatedLength = CalculateSequenceLengthOfEditorObject();
-		if(!FMath::IsNearlyEqual(CurrentCalculatedLength, AnimMontage->GetPlayLength(), UE_KINDA_SMALL_NUMBER))
+		float NewSequenceLength = CalculateSequenceLengthOfEditorObject();
+		if (NewSequenceLength != AnimMontage->GetPlayLength())
 		{
-			ClampToEndTime(CurrentCalculatedLength);
+			ClampToEndTime(NewSequenceLength);
 
 			RefreshSectionTimes();
 
-			AnimMontage->SetCompositeLength(CurrentCalculatedLength);
+			AnimMontage->SetCompositeLength(NewSequenceLength);
 
 			// Reset view if we changed length (note: has to be done after ->SetCompositeLength)!
 			UpdateRange();
@@ -312,7 +312,7 @@ void FAnimModel_AnimMontage::ToggleSectionTimingDisplay()
 	bSectionTimingEnabled = !bSectionTimingEnabled;
 }
 
-void FAnimModel_AnimMontage::OnDataModelChanged(const EAnimDataModelNotifyType& NotifyType, IAnimationDataModel* Model, const FAnimDataModelNotifPayload& PayLoad)
+void FAnimModel_AnimMontage::OnDataModelChanged(const EAnimDataModelNotifyType& NotifyType, UAnimDataModel* Model, const FAnimDataModelNotifPayload& PayLoad)
 {
 	NotifyCollector.Handle(NotifyType);
 
diff --git a/Engine/Source/Editor/Persona/Private/AnimTimeline/AnimModel_AnimMontage.h b/Engine/Source/Editor/Persona/Private/AnimTimeline/AnimModel_AnimMontage.h
--- a/Engine/Source/Editor/Persona/Private/AnimTimeline/AnimModel_AnimMontage.h
+++ b/Engine/Source/Editor/Persona/Private/AnimTimeline/AnimModel_AnimMontage.h
@@ -62,7 +62,7 @@ public:
 	void ToggleSectionTimingDisplay();
 
 protected:	
-	virtual void OnDataModelChanged(const EAnimDataModelNotifyType& NotifyType, IAnimationDataModel* Model, const FAnimDataModelNotifPayload& PayLoad) override;
+	virtual void OnDataModelChanged(const EAnimDataModelNotifyType& NotifyType, UAnimDataModel* Model, const FAnimDataModelNotifPayload& PayLoad) override;
 
 private:
 	/** The anim montage we wrap */
diff --git a/Engine/Source/Editor/Persona/Private/AnimTimeline/AnimModel_AnimSequenceBase.cpp b/Engine/Source/Editor/Persona/Private/AnimTimeline/AnimModel_AnimSequenceBase.cpp
--- a/Engine/Source/Editor/Persona/Private/AnimTimeline/AnimModel_AnimSequenceBase.cpp
+++ b/Engine/Source/Editor/Persona/Private/AnimTimeline/AnimModel_AnimSequenceBase.cpp
@@ -311,7 +311,7 @@ void FAnimModel_AnimSequenceBase::RefreshAttributeTracks()
 	}
 }
 
-void FAnimModel_AnimSequenceBase::OnDataModelChanged(const EAnimDataModelNotifyType& NotifyType, IAnimationDataModel* Model, const FAnimDataModelNotifPayload& PayLoad)
+void FAnimModel_AnimSequenceBase::OnDataModelChanged(const EAnimDataModelNotifyType& NotifyType, UAnimDataModel* Model, const FAnimDataModelNotifPayload& PayLoad)
 {
 	NotifyCollector.Handle(NotifyType);
 
diff --git a/Engine/Source/Editor/Persona/Private/AnimTimeline/AnimModel_AnimSequenceBase.h b/Engine/Source/Editor/Persona/Private/AnimTimeline/AnimModel_AnimSequenceBase.h
--- a/Engine/Source/Editor/Persona/Private/AnimTimeline/AnimModel_AnimSequenceBase.h
+++ b/Engine/Source/Editor/Persona/Private/AnimTimeline/AnimModel_AnimSequenceBase.h
@@ -59,8 +59,8 @@ protected:
 	/** Refresh attribute tracks */
 	void RefreshAttributeTracks();
 
-	/** Callback for any change made to the IAnimationDataModel embedded in the AnimSequenceBase instance this represents */
-	virtual void OnDataModelChanged(const EAnimDataModelNotifyType& NotifyType, IAnimationDataModel* Model, const FAnimDataModelNotifPayload& PayLoad);
+	/** Callback for any change made to the UAnimDataModel embedded in the AnimSequenceBase instance this represents */
+	virtual void OnDataModelChanged(const EAnimDataModelNotifyType& NotifyType, UAnimDataModel* Model, const FAnimDataModelNotifPayload& PayLoad);
 
 private:
 	/** UI handlers */
diff --git a/Engine/Source/Editor/Persona/Private/AnimTimeline/AnimTimelineTrack_FloatCurve.cpp b/Engine/Source/Editor/Persona/Private/AnimTimeline/AnimTimelineTrack_FloatCurve.cpp
--- a/Engine/Source/Editor/Persona/Private/AnimTimeline/AnimTimelineTrack_FloatCurve.cpp
+++ b/Engine/Source/Editor/Persona/Private/AnimTimeline/AnimTimelineTrack_FloatCurve.cpp
@@ -19,7 +19,7 @@
 #include "Widgets/Input/SComboButton.h"
 #include "Widgets/Layout/SBox.h"
 #include "AnimModel_AnimSequenceBase.h"
-#include "Animation/AnimData/IAnimationDataModel.h"
+#include "Animation/AnimData/AnimDataModel.h"
 #include "SAnimOutlinerItem.h"
 
 #define LOCTEXT_NAMESPACE "FAnimTimelineTrack_FloatCurve"
@@ -42,8 +42,7 @@ TSharedRef<SWidget> FAnimTimelineTrack_FloatCurve::MakeTimelineWidgetContainer()
 	// zoom to fit now we have a view
 	CurveEditor->ZoomToFit(EAxisList::Y);
 
-	FLinearColor CurveColor = FloatCurve->GetCurveTypeFlag(AACF_Metadata) ? FloatCurve->GetColor().Desaturate(0.25f) : FloatCurve->GetColor().Desaturate(0.75f);
-	auto ColorLambda = [this, CurveColor]()
+	auto ColorLambda = [this]()
 	{
 		if(GetModel()->IsTrackSelected(AsShared()))
 		{
@@ -51,14 +50,14 @@ TSharedRef<SWidget> FAnimTimelineTrack_FloatCurve::MakeTimelineWidgetContainer()
 		}
 		else
 		{
-			return CurveColor;
+			return FloatCurve->GetCurveTypeFlag(AACF_Metadata) ? FloatCurve->GetColor().Desaturate(0.25f) : FloatCurve->GetColor().Desaturate(0.75f);
 		}
 	};
 
 	return
 		SAssignNew(TimelineWidgetContainer, SBorder)
 		.Padding(0.0f)
-		.BorderImage(FloatCurve->GetCurveTypeFlag(AACF_Metadata) ? FAppStyle::GetBrush("Sequencer.Section.SelectedSectionOverlay") : FAppStyle::GetBrush("AnimTimeline.Outliner.DefaultBorder"))
+		.BorderImage_Lambda([this](){ return FloatCurve->GetCurveTypeFlag(AACF_Metadata) ? FAppStyle::GetBrush("Sequencer.Section.SelectedSectionOverlay") : FAppStyle::GetBrush("AnimTimeline.Outliner.DefaultBorder"); })
 		.BorderBackgroundColor_Lambda(ColorLambda)
 		[
 			CurveWidget
@@ -255,13 +254,14 @@ void FAnimTimelineTrack_FloatCurve::OnCommitCurveName(const FText& InText, EText
           
 			FAnimationCurveIdentifier NewCurveId(NewSmartName, ERawCurveTrackTypes::RCT_Float);
 			Controller.RenameCurve(CurveId, NewCurveId);
+			Controller.RemoveBoneTracksMissingFromSkeleton(AnimSequenceBase->GetSkeleton());           
 		}
 	}
 }
 
 FText FAnimTimelineTrack_FloatCurve::GetLabel() const
 {
-	return FAnimTimelineTrack_FloatCurve::GetFloatCurveName(GetModel(), CurveName);
+	return FAnimTimelineTrack_FloatCurve::GetFloatCurveName(GetModel(), FloatCurve->Name);
 }
 
 FText FAnimTimelineTrack_FloatCurve::GetFloatCurveName(const TSharedRef<FAnimModel>& InModel, const FSmartName& InSmartName)
@@ -303,7 +303,11 @@ void FAnimTimelineTrack_FloatCurve::AddCurveTrackButton(TSharedPtr<SHorizontalBo
 		PersonaUtils::MakeTrackButton(LOCTEXT("EditCurveButtonText", "Curve"), FOnGetContent::CreateSP(this, &FAnimTimelineTrack_FloatCurve::BuildCurveTrackMenu), MakeAttributeSP(this, &FAnimTimelineTrack_FloatCurve::IsHovered))
 	];
 
-	FLinearColor CurveColor = FloatCurve->Color;
+	auto GetValue = [this]()
+	{
+		return FloatCurve->Color;
+	};
+
 	auto SetValue = [this](FLinearColor InNewColor)
 	{
 		UAnimSequenceBase* AnimSequenceBase = GetModel()->GetAnimSequenceBase();
@@ -316,11 +320,11 @@ void FAnimTimelineTrack_FloatCurve::AddCurveTrackButton(TSharedPtr<SHorizontalBo
 		}
 	};
 
-	auto OnGetMenuContent = [this, CurveColor, SetValue]()
+	auto OnGetMenuContent = [this, GetValue, SetValue]()
 	{
 		// Open a color picker
 		return SNew(SColorPicker)
-			.TargetColorAttribute(CurveColor)
+			.TargetColorAttribute_Lambda(GetValue)
 			.UseAlpha(false)
 			.DisplayInlineVersion(true)
 			.OnColorCommitted_Lambda(SetValue)
@@ -353,7 +357,7 @@ void FAnimTimelineTrack_FloatCurve::AddCurveTrackButton(TSharedPtr<SHorizontalBo
 		.ButtonContent()
 		[
 			SNew(SColorBlock)
-			.Color(CurveColor)
+			.Color_Lambda(GetValue)
 			.ShowBackgroundForAlpha(false)
 			.AlphaDisplayMode(EColorBlockAlphaDisplayMode::Ignore)
 			.Size(FVector2D(OutlinerRightPadding - 2.0f, OutlinerRightPadding))
diff --git a/Engine/Source/Editor/Persona/Private/AnimationEditorViewportClient.cpp b/Engine/Source/Editor/Persona/Private/AnimationEditorViewportClient.cpp
--- a/Engine/Source/Editor/Persona/Private/AnimationEditorViewportClient.cpp
+++ b/Engine/Source/Editor/Persona/Private/AnimationEditorViewportClient.cpp
@@ -632,7 +632,7 @@ void FAnimationViewportClient::DrawUVsForMesh(FViewport* InViewport, FCanvas* In
 
 	TArray<FVector2D> SelectedEdgeTexCoords; //No functionality in Persona for this (yet?)
 
-	DrawUVs(InViewport, InCanvas, InTextYPos, LODLevel, UVChannelToDraw, SelectedEdgeTexCoords, nullptr, &PreviewMeshComponent->GetSkeletalMeshRenderData()->LODRenderData[LODLevel] );
+	DrawUVs(InViewport, InCanvas, InTextYPos, LODLevel, UVChannelToDraw, SelectedEdgeTexCoords, NULL, &PreviewMeshComponent->GetSkeletalMeshRenderData()->LODRenderData[LODLevel] );
 }
 
 void FAnimationViewportClient::Tick(float DeltaSeconds) 
@@ -854,7 +854,7 @@ bool FAnimationViewportClient::ShouldDisplayAdditiveScaleErrorMessage() const
 	{
 		if (AnimSequence->IsValidAdditive() && AnimSequence->RefPoseSeq)
 		{
-			FGuid AnimSeqGuid = AnimSequence->RefPoseSeq->GetDataModel()->GenerateGuid();
+			FGuid AnimSeqGuid = AnimSequence->RefPoseSeq->GetRawDataGuid();
 			if (RefPoseGuid != AnimSeqGuid)
 			{
 				RefPoseGuid = AnimSeqGuid;
@@ -980,7 +980,7 @@ FText FAnimationViewportClient::GetDisplayInfo(bool bDisplayAllInfo) const
 	{
 		if (!PreviewMeshComponent->CheckIfBoundsAreCorrrect())
 		{
-			if( PreviewMeshComponent->GetPhysicsAsset() == nullptr )
+			if( PreviewMeshComponent->GetPhysicsAsset() == NULL )
 			{
 				TextValue = ConcatenateLine(TextValue, LOCTEXT("NeedToSetupPhysicsAssetForAccurateBounds", "<AnimViewport.WarningText>You may need to setup Physics Asset to use more accurate bounds</>"));
 			}
@@ -991,7 +991,7 @@ FText FAnimationViewportClient::GetDisplayInfo(bool bDisplayAllInfo) const
 		}
 	}
 
-	if (PreviewMeshComponent != nullptr && PreviewMeshComponent->MeshObject != nullptr)
+	if (PreviewMeshComponent != NULL && PreviewMeshComponent->MeshObject != NULL)
 	{
 		if (bDisplayAllInfo)
 		{
@@ -1202,7 +1202,7 @@ FText FAnimationViewportClient::GetDisplayInfo(bool bDisplayAllInfo) const
 
 	if (const UPoseAsset* PoseAsset = Cast<UPoseAsset>(GetAnimPreviewScene()->GetPreviewAnimationAsset()))
 	{
-		if (PoseAsset->SourceAnimation && PoseAsset->SourceAnimation->GetDataModel()->GenerateGuid() != PoseAsset->SourceAnimationRawDataGUID)
+		if (PoseAsset->SourceAnimation && PoseAsset->SourceAnimation->GetRawDataGuid() != PoseAsset->SourceAnimationRawDataGUID)
 		{
 			TextValue = ConcatenateLine(TextValue, LOCTEXT("PoseAssetOutOfDateWarning", "<AnimViewport.WarningText>Poses are out-of-sync with the source animation. To update them click \"Update Source\"</>"));
 		}
@@ -2356,7 +2356,7 @@ void FAnimationViewportClient::SetupViewForRendering( FSceneViewFamily& ViewFami
 
 	// Cache screen size
 	UDebugSkelMeshComponent* PreviewMeshComponent = GetAnimPreviewScene()->GetPreviewMeshComponent();
-	if (PreviewMeshComponent != nullptr && PreviewMeshComponent->MeshObject != nullptr)
+	if (PreviewMeshComponent != NULL && PreviewMeshComponent->MeshObject != NULL)
 	{
 		const FBoxSphereBounds& SkelBounds = PreviewMeshComponent->Bounds;
 		CachedScreenSize = ComputeBoundsScreenSize(SkelBounds.Origin, SkelBounds.SphereRadius, View);
diff --git a/Engine/Source/Editor/Persona/Private/SAnimSequenceCurveEditor.cpp b/Engine/Source/Editor/Persona/Private/SAnimSequenceCurveEditor.cpp
--- a/Engine/Source/Editor/Persona/Private/SAnimSequenceCurveEditor.cpp
+++ b/Engine/Source/Editor/Persona/Private/SAnimSequenceCurveEditor.cpp
@@ -17,7 +17,7 @@
 #include "Editor.h"
 #include "Framework/MultiBox/MultiBoxBuilder.h"
 #include "Styling/AppStyle.h"
-#include "Animation/AnimData/IAnimationDataModel.h"
+#include "Animation/AnimData/AnimDataModel.h"
 
 #define LOCTEXT_NAMESPACE "SAnimSequenceCurveEditor"
 
@@ -122,7 +122,7 @@ void FRichCurveEditorModelNamed::CurveHasChanged()
 	}
 }
 
-void FRichCurveEditorModelNamed::OnModelHasChanged(const EAnimDataModelNotifyType& NotifyType, IAnimationDataModel* Model, const FAnimDataModelNotifPayload& Payload)
+void FRichCurveEditorModelNamed::OnModelHasChanged(const EAnimDataModelNotifyType& NotifyType, UAnimDataModel* Model, const FAnimDataModelNotifPayload& Payload)
 {
 	NotifyCollector.Handle(NotifyType);
 
@@ -367,7 +367,7 @@ SAnimSequenceCurveEditor::~SAnimSequenceCurveEditor()
 {
 	if(AnimSequence)
 	{
-		AnimSequence->GetDataModel()->GetModifiedEvent().RemoveAll(this);
+		AnimSequence->GetDataModel()->GetModifiedEvent().RemoveAll(this);		
 	}
 }
 
@@ -438,7 +438,7 @@ void SAnimSequenceCurveEditor::Construct(const FArguments& InArgs, const TShared
 	];
 }
 
-void SAnimSequenceCurveEditor::OnModelHasChanged(const EAnimDataModelNotifyType& NotifyType, IAnimationDataModel* Model, const FAnimDataModelNotifPayload& Payload)
+void SAnimSequenceCurveEditor::OnModelHasChanged(const EAnimDataModelNotifyType& NotifyType, UAnimDataModel* Model, const FAnimDataModelNotifPayload& Payload)
 {
 	auto StopEditingCurve = [this, NotifyType, &Payload, Model]()
 	{
diff --git a/Engine/Source/Editor/Persona/Private/SAnimSequenceCurveEditor.h b/Engine/Source/Editor/Persona/Private/SAnimSequenceCurveEditor.h
--- a/Engine/Source/Editor/Persona/Private/SAnimSequenceCurveEditor.h
+++ b/Engine/Source/Editor/Persona/Private/SAnimSequenceCurveEditor.h
@@ -35,7 +35,7 @@ public:
 	virtual void SetKeyAttributes(TArrayView<const FKeyHandle> InKeys, TArrayView<const FKeyAttributes> InAttributes, EPropertyChangeType::Type ChangeType = EPropertyChangeType::Unspecified) override;
 
 	void CurveHasChanged();
-	void OnModelHasChanged(const EAnimDataModelNotifyType& NotifyType, IAnimationDataModel* Model, const FAnimDataModelNotifPayload& Payload);
+	void OnModelHasChanged(const EAnimDataModelNotifyType& NotifyType, UAnimDataModel* Model, const FAnimDataModelNotifPayload& Payload);
 	void UpdateCachedCurve();
 
 	FSmartName Name;
@@ -72,7 +72,8 @@ class SAnimSequenceCurveEditor : public IAnimSequenceCurveEditor
 	virtual void RemoveCurve(const FSmartName& InName, ERawCurveTrackTypes InType, int32 InCurveIndex) override;
 	virtual void ZoomToFit() override;
 	
-	void OnModelHasChanged(const EAnimDataModelNotifyType& NotifyType, IAnimationDataModel* Model, const FAnimDataModelNotifPayload& Payload);
+	void OnModelHasChanged(const EAnimDataModelNotifyType& NotifyType, UAnimDataModel* Model, const FAnimDataModelNotifPayload& Payload);
+
 private:
 	// Build the toolbar for this curve editor
 	TSharedRef<SWidget> MakeToolbar(TSharedRef<SCurveEditorPanel> InEditorPanel);
diff --git a/Engine/Source/Editor/Persona/Private/SAnimationScrubPanel.cpp b/Engine/Source/Editor/Persona/Private/SAnimationScrubPanel.cpp
--- a/Engine/Source/Editor/Persona/Private/SAnimationScrubPanel.cpp
+++ b/Engine/Source/Editor/Persona/Private/SAnimationScrubPanel.cpp
@@ -10,7 +10,7 @@
 #include "ScopedTransaction.h"
 #include "Animation/BlendSpace.h"
 #include "AnimationEditorPreviewScene.h"
-#include "Animation/AnimData/IAnimationDataModel.h"
+#include "Animation/AnimData/AnimDataModel.h"
 #include "Animation/AnimSequenceHelpers.h"
 
 #define LOCTEXT_NAMESPACE "AnimationScrubPanel"
@@ -373,7 +373,7 @@ bool SAnimationScrubPanel::DoesSyncViewport() const
 {
 	UAnimSingleNodeInstance* PreviewInstance = GetPreviewInstance();
 
-	return (( LockedSequence==nullptr && PreviewInstance ) || ( LockedSequence && PreviewInstance && PreviewInstance->GetCurrentAsset() == LockedSequence ));
+	return (( LockedSequence==NULL && PreviewInstance ) || ( LockedSequence && PreviewInstance && PreviewInstance->GetCurrentAsset() == LockedSequence ));
 }
 
 void SAnimationScrubPanel::Tick( const FGeometry& AllottedGeometry, const double InCurrentTime, const float InDeltaTime )
@@ -387,7 +387,7 @@ void SAnimationScrubPanel::Tick( const FGeometry& AllottedGeometry, const double
 class UAnimSingleNodeInstance* SAnimationScrubPanel::GetPreviewInstance() const
 {
 	UDebugSkelMeshComponent* PreviewMeshComponent = GetPreviewScene()->GetPreviewMeshComponent();
-	return PreviewMeshComponent && PreviewMeshComponent->IsPreviewOn()? PreviewMeshComponent->PreviewInstance : nullptr;
+	return PreviewMeshComponent && PreviewMeshComponent->IsPreviewOn()? PreviewMeshComponent->PreviewInstance : NULL;
 }
 
 float SAnimationScrubPanel::GetScrubValue() const
@@ -424,20 +424,20 @@ UAnimInstance* SAnimationScrubPanel::GetAnimInstanceWithBlueprint() const
 	{
 		UAnimInstance* Instance = DebugComponent->GetAnimInstance();
 
-		if ((Instance != nullptr) && (Instance->GetClass()->ClassGeneratedBy != nullptr))
+		if ((Instance != NULL) && (Instance->GetClass()->ClassGeneratedBy != NULL))
 		{
 			return Instance;
 		}
 	}
 
-	return nullptr;
+	return NULL;
 }
 
 bool SAnimationScrubPanel::GetAnimBlueprintDebugData(UAnimInstance*& Instance, FAnimBlueprintDebugData*& DebugInfo) const
 {
 	Instance = GetAnimInstanceWithBlueprint();
 
-	if (Instance != nullptr)
+	if (Instance != NULL)
 	{
 		// Avoid updating the instance if we're replaying the past
 		if (UAnimBlueprintGeneratedClass* AnimBlueprintClass = Cast<UAnimBlueprintGeneratedClass>(Instance->GetClass()))
diff --git a/Engine/Source/Editor/Persona/Private/SAnimationSequenceBrowser.cpp b/Engine/Source/Editor/Persona/Private/SAnimationSequenceBrowser.cpp
--- a/Engine/Source/Editor/Persona/Private/SAnimationSequenceBrowser.cpp
+++ b/Engine/Source/Editor/Persona/Private/SAnimationSequenceBrowser.cpp
@@ -469,7 +469,7 @@ SAnimationSequenceBrowser::~SAnimationSequenceBrowser()
 
 	if(ViewportClient.IsValid())
 	{
-		ViewportClient->Viewport = nullptr;
+		ViewportClient->Viewport = NULL;
 	}
 }
 
@@ -700,7 +700,7 @@ void SAnimationSequenceBrowser::GetSelectedPackages(const TArray<FAssetData>& As
 {
 	for (int32 AssetIdx = 0; AssetIdx < Assets.Num(); ++AssetIdx)
 	{
-		UPackage* Package = FindPackage(nullptr, *Assets[AssetIdx].PackageName.ToString());
+		UPackage* Package = FindPackage(NULL, *Assets[AssetIdx].PackageName.ToString());
 
 		if ( Package )
 		{
@@ -1014,7 +1014,7 @@ void SAnimationSequenceBrowser::Construct(const FArguments& InArgs, const TShare
 	// Making Skeleton to be private, and now GET_MEMBER_NAME_CHECKED doesn't work
 	AssetRegistryTagsToIgnore.Add(TEXT("Skeleton"));
 	PRAGMA_DISABLE_DEPRECATION_WARNINGS
-	AssetRegistryTagsToIgnore.Add(TEXT("SequenceLength"));
+	AssetRegistryTagsToIgnore.Add(GET_MEMBER_NAME_CHECKED(UAnimSequenceBase, SequenceLength));
 	PRAGMA_ENABLE_DEPRECATION_WARNINGS
 	AssetRegistryTagsToIgnore.Add(GET_MEMBER_NAME_CHECKED(UAnimSequenceBase, RateScale));
 }
@@ -1066,7 +1066,7 @@ FReply SAnimationSequenceBrowser::OnMouseDownHistory( const FGeometry& MyGeometr
 
 TSharedRef<SWidget> SAnimationSequenceBrowser::CreateHistoryMenu(bool bInBackHistory) const
 {
-	FMenuBuilder MenuBuilder(true, nullptr);
+	FMenuBuilder MenuBuilder(true, NULL);
 	if(bInBackHistory)
 	{
 		int32 HistoryIdx = CurrentAssetHistoryIndex - 1;
diff --git a/Engine/Source/Editor/SequenceRecorder/Private/AnimationRecorder.cpp b/Engine/Source/Editor/SequenceRecorder/Private/AnimationRecorder.cpp
--- a/Engine/Source/Editor/SequenceRecorder/Private/AnimationRecorder.cpp
+++ b/Engine/Source/Editor/SequenceRecorder/Private/AnimationRecorder.cpp
@@ -244,12 +244,10 @@ void FAnimationRecorder::StartRecord(USkeletalMeshComponent* Component, UAnimSeq
 	LastFrame = 0;
 
 	IAnimationDataController& Controller = AnimationObject->GetController();
-	Controller.SetModel(AnimationObject->GetDataModelInterface());
+	Controller.SetModel(AnimationObject->GetDataModel());
 
 	Controller.OpenBracket(LOCTEXT("StartRecord_Bracket", "Starting Animation Recording"));
 
-	Controller.InitializeModel();
-
 	const bool bKeepNotifiesAndCurves = CVarKeepNotifyAndCurvesOnAnimationRecord->GetInt() == 0 ? false : true;
 	if (bKeepNotifiesAndCurves)
 	{
@@ -361,8 +359,8 @@ UAnimSequence* FAnimationRecorder::StopRecord(bool bShowMessage)
 		int32 NumKeys = LastFrame.Value + 1;
 
 		// can't use TimePassed. That is just total time that has been passed, not necessarily match with frame count
+		Controller.SetPlayLength( (NumKeys>1) ? RecordingRate.AsSeconds(LastFrame): RecordingRate.AsSeconds(1) );
 		Controller.SetFrameRate(RecordingRate);
-		Controller.SetNumberOfFrames( FMath::Max(LastFrame.Value,1) );
 
 		ProcessNotifies();
 
@@ -523,7 +521,7 @@ UAnimSequence* FAnimationRecorder::StopRecord(bool bShowMessage)
 			FSavePackageArgs SaveArgs;
 			SaveArgs.TopLevelFlags = RF_Standalone;
 			SaveArgs.SaveFlags = SAVE_NoError;
-			UPackage::SavePackage(Package, nullptr, *PackageFileName, SaveArgs);
+			UPackage::SavePackage(Package, NULL, *PackageFileName, SaveArgs);
 
 			ElapsedTime = FPlatformTime::Seconds() - StartTime;
 			UE_LOG(LogAnimation, Log, TEXT("Animation Recorder saved %s in %0.2f seconds"), *PackageName, ElapsedTime);
@@ -563,7 +561,7 @@ UAnimSequence* FAnimationRecorder::StopRecord(bool bShowMessage)
 			FAssetRegistryModule::AssetCreated(AnimationObject);
 		}
 
-		AnimationObject = nullptr;
+		AnimationObject = NULL;
 		PreviousSpacesBases.Empty();
 		PreviousAnimCurves.Empty();
 
diff --git a/Engine/Source/Editor/UnrealEd/Private/Animation/DebugSkelMeshComponent.cpp b/Engine/Source/Editor/UnrealEd/Private/Animation/DebugSkelMeshComponent.cpp
--- a/Engine/Source/Editor/UnrealEd/Private/Animation/DebugSkelMeshComponent.cpp
+++ b/Engine/Source/Editor/UnrealEd/Private/Animation/DebugSkelMeshComponent.cpp
@@ -31,7 +31,7 @@ UDebugSkelMeshComponent::UDebugSkelMeshComponent(const FObjectInitializer& Objec
 	: Super(ObjectInitializer)
 {
 	bDrawMesh = true;
-	PreviewInstance = nullptr;
+	PreviewInstance = NULL;
 	bDisplayRawAnimation = false;
 	bDisplayNonRetargetedPose = false;
 
@@ -422,9 +422,9 @@ bool UDebugSkelMeshComponent::IsTrackingAttachedLOD() const
 
 FPrimitiveSceneProxy* UDebugSkelMeshComponent::CreateSceneProxy()
 {
-	FDebugSkelMeshSceneProxy* Result = nullptr;
+	FDebugSkelMeshSceneProxy* Result = NULL;
 	ERHIFeatureLevel::Type SceneFeatureLevel = GetWorld()->FeatureLevel;
-	FSkeletalMeshRenderData* SkelMeshRenderData = SkeletalMesh ? SkeletalMesh->GetResourceForRendering() : nullptr;
+	FSkeletalMeshRenderData* SkelMeshRenderData = SkeletalMesh ? SkeletalMesh->GetResourceForRendering() : NULL;
 
 	// only create a scene proxy for rendering if
 	// properly initialized
@@ -446,7 +446,7 @@ bool UDebugSkelMeshComponent::ShouldRenderSelected() const
 
 bool UDebugSkelMeshComponent::IsPreviewOn() const
 {
-	return (PreviewInstance != nullptr) && (PreviewInstance == AnimScriptInstance);
+	return (PreviewInstance != NULL) && (PreviewInstance == AnimScriptInstance);
 }
 
 FString UDebugSkelMeshComponent::GetPreviewText() const
@@ -493,7 +493,7 @@ void UDebugSkelMeshComponent::InitAnim(bool bForceReinit)
 		if (!SkeletalMesh->GetSkeleton()->IsCompatible(PreviewInstance->GetCurrentAsset()->GetSkeleton()))
 		{
 			// if it doesn't match, just clear it
-			PreviewInstance->SetAnimationAsset(nullptr);
+			PreviewInstance->SetAnimationAsset(NULL);
 		}
 	}
 
@@ -506,8 +506,8 @@ void UDebugSkelMeshComponent::InitAnim(bool bForceReinit)
 
 	Super::InitAnim(bForceReinit);
 
-	// if PreviewInstance is nullptr, create here once
-	if (PreviewInstance == nullptr)
+	// if PreviewInstance is NULL, create here once
+	if (PreviewInstance == NULL)
 	{
 		PreviewInstance = NewObject<UAnimPreviewInstance>(this);
 		check(PreviewInstance);
@@ -518,7 +518,7 @@ void UDebugSkelMeshComponent::InitAnim(bool bForceReinit)
 
 	// if anim script instance is null because it's not playing a blueprint, set to PreviewInstnace by default
 	// that way if user would like to modify bones or do extra stuff, it will work
-	if (AnimScriptInstance == nullptr)
+	if (AnimScriptInstance == NULL)
 	{
 		AnimScriptInstance = PreviewInstance;
 		AnimScriptInstance->InitializeAnimation();
@@ -571,7 +571,7 @@ void UDebugSkelMeshComponent::EnablePreview(bool bEnable, UAnimationAsset* Previ
 		}
 		else if (IsPreviewOn())
 		{
-			if (PreviewInstance->GetCurrentAsset() == PreviewAsset || PreviewAsset == nullptr)
+			if (PreviewInstance->GetCurrentAsset() == PreviewAsset || PreviewAsset == NULL)
 			{
 				// now recover to saved AnimScriptInstance;
 				AnimScriptInstance = SavedAnimScriptInstance;
@@ -724,7 +724,7 @@ void UDebugSkelMeshComponent::GenSpaceBases(TArray<FTransform>& OutSpaceBases)
 void UDebugSkelMeshComponent::RefreshBoneTransforms(FActorComponentTickFunction* TickFunction)
 {
 	// Run regular update first so we get RequiredBones up to date.
-	Super::RefreshBoneTransforms(nullptr); // Pass nullptr so we force non threaded work
+	Super::RefreshBoneTransforms(NULL); // Pass NULL so we force non threaded work
 
 	// none of these code works if we don't have anim instance, so no reason to check it for every if
 	if (AnimScriptInstance && AnimScriptInstance->GetRequiredBones().IsValid())
@@ -810,7 +810,7 @@ void UDebugSkelMeshComponent::RefreshBoneTransforms(FActorComponentTickFunction*
 						AdditiveBasePose.SetBoneContainer(&BoneContainer);
 						
 						FAnimationPoseData AnimationPoseData(AdditiveBasePose, AdditiveCurve, AdditiveAttributes);
-						Sequence->GetAdditiveBasePose(AnimationPoseData, FAnimExtractContext(static_cast<double>(PreviewInstance->GetCurrentTime())));
+						Sequence->GetAdditiveBasePose(AnimationPoseData, FAnimExtractContext(PreviewInstance->GetCurrentTime()));
 						CSAdditiveBasePose.InitPose(AnimationPoseData.GetPose());
 					}
 
diff --git a/Engine/Source/Editor/UnrealEd/Private/AnimationEditorUtils.cpp b/Engine/Source/Editor/UnrealEd/Private/AnimationEditorUtils.cpp
--- a/Engine/Source/Editor/UnrealEd/Private/AnimationEditorUtils.cpp
+++ b/Engine/Source/Editor/UnrealEd/Private/AnimationEditorUtils.cpp
@@ -503,12 +503,6 @@ namespace AnimationEditorUtils
 				if(NewAsset)
 				{
 					NewAsset->SetSkeleton(Skeleton);
-
-					if (UAnimSequenceBase* SequenceBase = Cast<UAnimSequenceBase>(NewAsset))
-					{
-						SequenceBase->GetController().InitializeModel();
-					}
-					
 					if (SkeletalMesh)
 					{
 						NewAsset->SetPreviewMesh(SkeletalMesh);
diff --git a/Engine/Source/Editor/UnrealEd/Private/Factories/AnimSequenceFactory.cpp b/Engine/Source/Editor/UnrealEd/Private/Factories/AnimSequenceFactory.cpp
--- a/Engine/Source/Editor/UnrealEd/Private/Factories/AnimSequenceFactory.cpp
+++ b/Engine/Source/Editor/UnrealEd/Private/Factories/AnimSequenceFactory.cpp
@@ -72,9 +72,6 @@ UObject* UAnimSequenceFactory::FactoryCreateNew(UClass* Class, UObject* InParent
 	{
 		AnimSequence->SetSkeleton(TargetSkeleton);
 	}
-
-	AnimSequence->GetController().InitializeModel();
-	
 	if (PreviewSkeletalMesh)
 	{
 		AnimSequence->SetPreviewMesh(PreviewSkeletalMesh);
diff --git a/Engine/Source/Editor/UnrealEd/Private/Fbx/FbxAnimationExport.cpp b/Engine/Source/Editor/UnrealEd/Private/Fbx/FbxAnimationExport.cpp
--- a/Engine/Source/Editor/UnrealEd/Private/Fbx/FbxAnimationExport.cpp
+++ b/Engine/Source/Editor/UnrealEd/Private/Fbx/FbxAnimationExport.cpp
@@ -72,14 +72,6 @@ void FFbxExporter::ExportAnimSequenceToFbx(const UAnimSequence* AnimSeq,
 									 float AnimEndOffset,
 									 float AnimPlayRate,
 									 float StartTime)
-{
-	ExportAnimSequenceToFbx(AnimSeq, SkelMesh, BoneNodes, InAnimLayer,
-		AnimSeq->GetDataModel()->GetFrameRate().AsFrameTime(AnimStartOffset),
-		AnimSeq->GetDataModel()->GetFrameRate().AsFrameTime(AnimSeq->GetPlayLength() - AnimEndOffset),
-		AnimPlayRate, StartTime);
-}
-
-void FFbxExporter::ExportAnimSequenceToFbx(const UAnimSequence* AnimSeq, const USkeletalMesh* SkelMesh, TArray<FbxNode*>& BoneNodes, FbxAnimLayer* InAnimLayer, FFrameTime StartFrameTime, FFrameTime EndFrameTime, float FrameRateScale, float StartTime)
 {
 	// stack allocator for extracting curve
 	FMemMark Mark(FMemStack::Get());
@@ -123,7 +115,7 @@ void FFbxExporter::ExportAnimSequenceToFbx(const UAnimSequence* AnimSeq, const U
 		}
 	}
 
-	ExportCustomAnimCurvesToFbx(CustomCurveMap, AnimSeq, StartFrameTime, EndFrameTime, FrameRateScale, StartTime);
+	ExportCustomAnimCurvesToFbx(CustomCurveMap, AnimSeq, AnimStartOffset, AnimEndOffset, AnimPlayRate, StartTime);
 
 	TArray<const FAnimatedBoneAttribute*> CustomAttributes;
 
@@ -218,7 +210,7 @@ void FFbxExporter::ExportAnimSequenceToFbx(const UAnimSequence* AnimSeq, const U
 			Curve->KeyModifyBegin();
 		}
 
-		auto ExportLambda = [&](double AnimTime, FbxTime ExportTime, bool bLastKey) {
+		auto ExportLambda = [&](float AnimTime, FbxTime ExportTime, bool bLastKey) {
 			FTransform BoneAtom;
 			AnimSeq->GetBoneTransform(BoneAtom, BoneTrackIndex, AnimTime, true);
 			FbxAMatrix FbxMatrix = Converter.ConvertMatrix(BoneAtom.ToMatrixWithScale());
@@ -276,7 +268,7 @@ void FFbxExporter::ExportAnimSequenceToFbx(const UAnimSequence* AnimSeq, const U
 			}
 		};
 
-		IterateInsideAnimSequence(AnimSeq, StartFrameTime, EndFrameTime, FrameRateScale, StartTime, ExportLambda);
+		IterateInsideAnimSequence(AnimSeq, AnimStartOffset, AnimEndOffset, AnimPlayRate, StartTime, ExportLambda);
 
 		for (FbxAnimCurve* Curve : Curves)
 		{
@@ -298,15 +290,6 @@ void FFbxExporter::ExportAnimSequenceToFbx(const UAnimSequence* AnimSeq, const U
 
 void FFbxExporter::ExportCustomAnimCurvesToFbx(const TMap<FName, FbxAnimCurve*>& CustomCurves, const UAnimSequence* AnimSeq, 
 	float AnimStartOffset, float AnimEndOffset, float AnimPlayRate, float StartTime, float ValueScale)
-{
-	ExportCustomAnimCurvesToFbx(CustomCurves, AnimSeq,
-		AnimSeq->GetDataModel()->GetFrameRate().AsFrameTime(AnimStartOffset),
-		AnimSeq->GetDataModel()->GetFrameRate().AsFrameTime(AnimSeq->GetPlayLength() - AnimEndOffset),
-		AnimPlayRate, StartTime, ValueScale);
-}
-	
-void FFbxExporter::ExportCustomAnimCurvesToFbx(const TMap<FName, FbxAnimCurve*>& CustomCurves, const UAnimSequence* AnimSeq,
-	FFrameTime AnimStartOffset, FFrameTime AnimEndOffset, float FrameRateScale, float StartTime, float ValueScale /*= 1.f*/)
 {
 	// stack allocator for extracting curve
 	FMemMark Mark(FMemStack::Get());
@@ -356,7 +339,7 @@ void FFbxExporter::ExportCustomAnimCurvesToFbx(const TMap<FName, FbxAnimCurve*>&
 		}
 	};
 
-	IterateInsideAnimSequence(AnimSeq, AnimStartOffset, AnimEndOffset, FrameRateScale, StartTime, ExportLambda);
+	IterateInsideAnimSequence(AnimSeq, AnimStartOffset, AnimEndOffset, AnimPlayRate, StartTime, ExportLambda);
 
 	for (auto CustomCurve : CustomCurves)
 	{
@@ -366,21 +349,10 @@ void FFbxExporter::ExportCustomAnimCurvesToFbx(const TMap<FName, FbxAnimCurve*>&
 
 void FFbxExporter::IterateInsideAnimSequence(const UAnimSequence* AnimSeq, float AnimStartOffset, float AnimEndOffset, float AnimPlayRate, float StartTime, TFunctionRef<void(float, FbxTime, bool)> IterationLambda)
 {
-	IterateInsideAnimSequence(AnimSeq,
-		AnimSeq->GetDataModel()->GetFrameRate().AsFrameTime(AnimStartOffset),
-		AnimSeq->GetDataModel()->GetFrameRate().AsFrameTime(AnimSeq->GetPlayLength() - AnimEndOffset),
-		AnimPlayRate,
-		AnimPlayRate,
-		[IterationLambda](double T, FbxTime FT, bool B) -> void
-		{
-			IterationLambda((float)T, FT, B);
-		});
-}
-	
-void FFbxExporter::IterateInsideAnimSequence(const UAnimSequence* AnimSeq, FFrameTime StartFrameTime, FFrameTime EndFrameTime, float FrameRateScale, float StartTime, TFunctionRef<void(double, FbxTime, bool)> IterationLambda)
-{
-	const double TimePerKey = AnimSeq->GetDataModel()->GetFrameRate().AsInterval();
-	const double AnimTimeIncrement = TimePerKey * FrameRateScale;
+	float AnimTime = AnimStartOffset;
+	float AnimEndTime = (AnimSeq->GetPlayLength() - AnimEndOffset);
+	double TimePerKey = AnimSeq->GetDataModel()->GetFrameRate().AsInterval();
+	const float AnimTimeIncrement = TimePerKey * AnimPlayRate;
 	uint32 AnimFrameIndex = 0;
 
 	FbxTime ExportTime;
@@ -391,16 +363,15 @@ void FFbxExporter::IterateInsideAnimSequence(const UAnimSequence* AnimSeq, FFram
 
 	// Step through each frame and add custom curve data
 	bool bLastKey = false;
-	FFrameTime FrameTime = StartFrameTime;
-	const FFrameRate& FrameRate = AnimSeq->GetDataModel()->GetFrameRate();
 	while (!bLastKey)
 	{
-		bLastKey = FrameTime > EndFrameTime;
-		IterationLambda(FrameRate.AsSeconds(FrameTime), ExportTime, bLastKey);
+		bLastKey = (AnimTime + KINDA_SMALL_NUMBER) > AnimEndTime;
+
+		IterationLambda(AnimTime, ExportTime, bLastKey);
 
 		ExportTime += ExportTimeIncrement;
 		AnimFrameIndex++;
-		FrameTime += FFrameTime::FromDecimal(FrameRateScale);
+		AnimTime = AnimStartOffset + ((float)AnimFrameIndex * AnimTimeIncrement);
 	}
 }
 
@@ -455,9 +426,9 @@ void FFbxExporter::CorrectAnimTrackInterpolation( TArray<FbxNode*>& BoneNodes, F
 
 FbxNode* FFbxExporter::ExportAnimSequence( const UAnimSequence* AnimSeq, const USkeletalMesh* SkelMesh, bool bExportSkelMesh, const TCHAR* MeshName, FbxNode* ActorRootNode, const TArray<UMaterialInterface*>* OverrideMaterials /*= nullptr*/ )
 {
-	if( Scene == nullptr || AnimSeq == nullptr || SkelMesh == nullptr )
+	if( Scene == NULL || AnimSeq == NULL || SkelMesh == NULL )
 	{
- 		return nullptr;
+ 		return NULL;
 	}
 
 
@@ -483,14 +454,10 @@ FbxNode* FFbxExporter::ExportAnimSequence( const UAnimSequence* AnimSeq, const U
 			SkelMesh,
 			BoneNodes,
 			AnimLayer,
-			// Start frame to export
-			FFrameTime(0),
-			// Final frame to export
-			FFrameTime(AnimSeq->GetDataModel()->GetNumberOfFrames()),
-			// Frame rate scale
-			1.f,
-			// FBX StartTime
-			0.f);
+			0.f,		// AnimStartOffset
+			0.f,		// AnimEndOffset
+			1.f,		// AnimPlayRate
+			0.f);		// StartTime
 
 		CorrectAnimTrackInterpolation(BoneNodes, AnimLayer);
 	}
@@ -587,7 +554,7 @@ FbxNode* FFbxExporter::ExportAnimSequence( const UAnimSequence* AnimSeq, const U
 
 void FFbxExporter::ExportAnimSequencesAsSingle( USkeletalMesh* SkelMesh, const ASkeletalMeshActor* SkelMeshActor, const FString& ExportName, const TArray<UAnimSequence*>& AnimSeqList, const TArray<struct FAnimControlTrackKey>& TrackKeys )
 {
-	if (Scene == nullptr || SkelMesh == nullptr || AnimSeqList.Num() == 0 || AnimSeqList.Num() != TrackKeys.Num()) return;
+	if (Scene == NULL || SkelMesh == NULL || AnimSeqList.Num() == 0 || AnimSeqList.Num() != TrackKeys.Num()) return;
 
 	FbxNode* BaseNode = FbxNode::Create(Scene, Converter.ConvertToFbxString(ExportName));
 	Scene->GetRootNode()->AddChild(BaseNode);
@@ -621,8 +588,8 @@ void FFbxExporter::ExportAnimSequencesAsSingle( USkeletalMesh* SkelMesh, const A
 			SkelMesh,
 			BoneNodes,
 			AnimLayer,
-			AnimSeq->GetDataModel()->GetFrameRate().AsFrameTime(TrackKey.AnimStartOffset),
-			AnimSeq->GetDataModel()->GetFrameRate().AsFrameTime(AnimSeq->GetPlayLength() - TrackKey.AnimEndOffset),
+			TrackKey.AnimStartOffset,
+			TrackKey.AnimEndOffset,
 			TrackKey.AnimPlayRate,
 			CurrentStartTime);
 	}
@@ -646,7 +613,7 @@ void FFbxExporter::ExportAnimSequencesAsSingle( USkeletalMesh* SkelMesh, const A
  */
 void FFbxExporter::ExportMatineeGroup(class AMatineeActor* MatineeActor, USkeletalMeshComponent* SkeletalMeshComponent)
 {
-	if (Scene == nullptr || MatineeActor == nullptr || SkeletalMeshComponent == nullptr || MatineeActor->MatineeData->InterpLength == 0)
+	if (Scene == NULL || MatineeActor == NULL || SkeletalMeshComponent == NULL || MatineeActor->MatineeData->InterpLength == 0)
 	{
 		return;
 	}
