Convert "Unable to find bone index, animation track will be removed: X" warning into a display log entry (and updated overall API)

diff --git a/Engine/Source/Developer/AnimationDataController/Private/AnimDataController.cpp b/Engine/Source/Developer/AnimationDataController/Private/AnimDataController.cpp
--- a/Engine/Source/Developer/AnimationDataController/Private/AnimDataController.cpp
+++ b/Engine/Source/Developer/AnimationDataController/Private/AnimDataController.cpp
@@ -106,7 +106,7 @@ void UAnimDataController::CloseBracket(bool bShouldTransact /*= true*/)
 
 	if (BracketDepth == 0)
 	{
-		ReportError(LOCTEXT("NoExistingBracketError", "Attempt to close bracket that was not previously opened"));
+		Report(ELogVerbosity::Error, LOCTEXT("NoExistingBracketError", "Attempt to close bracket that was not previously opened"));
 		return;
 	}
 
@@ -264,7 +264,7 @@ void UAnimDataController::UpdateCurveNamesFromSkeleton(const USkeleton* Skeleton
 	}
 	else
 	{
-		ReportError(LOCTEXT("UpdateCurveInvalidSkeletonError", "Invalid USkeleton supplied"));
+		Report(ELogVerbosity::Error, LOCTEXT("UpdateCurveInvalidSkeletonError", "Invalid USkeleton supplied"));
 	}
 }
 
@@ -319,7 +319,7 @@ void UAnimDataController::FindOrAddCurveNamesOnSkeleton(USkeleton* Skeleton, ERa
 	}
 	else
 	{
-		ReportError(LOCTEXT("FindOrAddCurveInvalidSkeletonError", "Invalid USkeleton supplied "));
+		Report(ELogVerbosity::Error, LOCTEXT("FindOrAddCurveInvalidSkeletonError", "Invalid USkeleton supplied "));
 	}
 }
 
@@ -351,7 +351,7 @@ bool UAnimDataController::RemoveBoneTracksMissingFromSkeleton(const USkeleton* S
 			{				
 				// Remove track
 				TracksToBeRemoved.Add(Track.Name);
-				ReportWarningf(LOCTEXT("InvalidBoneIndexWarning", "Unable to find bone index, animation track will be removed: {0}"), FText::FromName(Track.Name));				
+				Reportf(ELogVerbosity::Display, LOCTEXT("InvalidBoneIndexWarning", "Unable to find bone index, animation track will be removed: {0}"), FText::FromName(Track.Name));				
 			}			
 		}
 
@@ -375,7 +375,7 @@ bool UAnimDataController::RemoveBoneTracksMissingFromSkeleton(const USkeleton* S
 	}
 	else
 	{
-		ReportError(LOCTEXT("InvalidSkeletonError", "Invalid USkeleton supplied"));
+		Report(ELogVerbosity::Error, LOCTEXT("InvalidSkeletonError", "Invalid USkeleton supplied"));
 	}
 
 	return false;
@@ -425,7 +425,7 @@ void UAnimDataController::UpdateAttributesFromSkeleton(const USkeleton* Skeleton
 	}
 	else
 	{
-		ReportError(LOCTEXT("InvalidSkeletonError", "Invalid USkeleton supplied"));
+		Report(ELogVerbosity::Error, LOCTEXT("InvalidSkeletonError", "Invalid USkeleton supplied"));
 	}
 }
 
@@ -1006,7 +1006,7 @@ bool UAnimDataController::SetCurveColor(const FAnimationCurveIdentifier& CurveId
 		}
 		else
 		{
-			ReportWarning(LOCTEXT("NonSupportedCurveColorSetWarning", "Changing curve color is currently only supported for float curves"));
+			Report(ELogVerbosity::Warning, LOCTEXT("NonSupportedCurveColorSetWarning", "Changing curve color is currently only supported for float curves"));
 		}
 	}
 	else
@@ -1049,7 +1049,7 @@ bool UAnimDataController::ScaleCurve(const FAnimationCurveIdentifier& CurveId, f
 	}
 	else
 	{
-		ReportWarning(LOCTEXT("NonSupportedCurveScalingWarning", "Scaling curves is currently only supported for float curves"));
+		Report(ELogVerbosity::Warning, LOCTEXT("NonSupportedCurveScalingWarning", "Scaling curves is currently only supported for float curves"));
 	}
 
 	return false;
@@ -1237,21 +1237,8 @@ void UAnimDataController::ResizePlayLength(float Length, float T0, float T1, boo
 	}
 }
 
-void UAnimDataController::ReportWarning(const FText& InMessage) const
-{
-	FString Message = InMessage.ToString();
-	if (Model != nullptr)
-	{
-		if (UPackage* Package = Cast<UPackage>(Model->GetOutermost()))
-		{
-			Message = FString::Printf(TEXT("%s : %s"), *Package->GetPathName(), *Message);
-		}
-	}
-
-	FScriptExceptionHandler::Get().HandleException(ELogVerbosity::Warning, *Message, *FString());
-}
 
-void UAnimDataController::ReportError(const FText& InMessage) const
+void UAnimDataController::Report(ELogVerbosity::Type InVerbosity, const FText& InMessage) const
 {
 	FString Message = InMessage.ToString();
 	if (Model != nullptr)
@@ -1262,7 +1249,7 @@ void UAnimDataController::ReportError(const FText& InMessage) const
 		}
 	}
 
-	FScriptExceptionHandler::Get().HandleException(ELogVerbosity::Error, *Message, *FString());
+	FScriptExceptionHandler::Get().HandleException(InVerbosity, *Message, *FString());
 }
 
 FString UAnimDataController::GetCurveTypeValueName(ERawCurveTrackTypes InType) const
@@ -1357,12 +1344,12 @@ int32 UAnimDataController::InsertBoneTrack(FName BoneName, int32 DesiredIndex, b
 				}
 				else
 				{
-					ReportError(LOCTEXT("UnableToGetOuterSkeletonError", "Unable to retrieve Skeleton for outer Animation Sequence"));
+					Report(ELogVerbosity::Error, LOCTEXT("UnableToGetOuterSkeletonError", "Unable to retrieve Skeleton for outer Animation Sequence"));
 				}
 			}
 			else
 			{
-				ReportError(LOCTEXT("UnableToGetOuterAnimSequenceError", "Unable to retrieve outer Animation Sequence"));
+				Report(ELogVerbosity::Error, LOCTEXT("UnableToGetOuterAnimSequenceError", "Unable to retrieve outer Animation Sequence"));
 			}
 
 				FAnimationTrackAddedPayload Payload;
@@ -1830,7 +1817,7 @@ bool UAnimDataController::AddAttribute(const FAnimationAttributeIdentifier& Attr
 	}
 	else
 	{
-		ReportError(LOCTEXT("InvalidAttributeIdentifier", "Invalid attribute identifier provided"));
+		Report(ELogVerbosity::Error, LOCTEXT("InvalidAttributeIdentifier", "Invalid attribute identifier provided"));
 	}
 
 	return false;
@@ -1867,7 +1854,7 @@ bool UAnimDataController::RemoveAttribute(const FAnimationAttributeIdentifier& A
 	}
 	else
 	{
-		ReportError(LOCTEXT("InvalidAttributeIdentifier", "Invalid attribute identifier provided"));
+		Report(ELogVerbosity::Error, LOCTEXT("InvalidAttributeIdentifier", "Invalid attribute identifier provided"));
 	}
 
 	return false;
@@ -1977,12 +1964,12 @@ bool UAnimDataController::SetAttributeKey_Internal(const FAnimationAttributeIden
 		}
 		else
 		{
-			ReportError(LOCTEXT("InvalidAttributeKey", "Invalid attribute key value provided"));
+			Report(ELogVerbosity::Error, LOCTEXT("InvalidAttributeKey", "Invalid attribute key value provided"));
 		}
 	}
 	else
 	{
-		ReportError(LOCTEXT("InvalidAttributeIdentifier", "Invalid attribute identifier provided"));
+		Report(ELogVerbosity::Error, LOCTEXT("InvalidAttributeIdentifier", "Invalid attribute identifier provided"));
 	}
 
 	return false;
@@ -2037,7 +2024,7 @@ bool UAnimDataController::SetAttributeKeys_Internal(const FAnimationAttributeIde
 	}
 	else
 	{
-		ReportError(LOCTEXT("InvalidAttributeIdentifier", "Invalid attribute identifier provided"));
+		Report(ELogVerbosity::Error, LOCTEXT("InvalidAttributeIdentifier", "Invalid attribute identifier provided"));
 	}
 
 	return false;
@@ -2074,7 +2061,7 @@ bool UAnimDataController::RemoveAttributeKey(const FAnimationAttributeIdentifier
 			}
 			else
 			{
-				ReportWarning(LOCTEXT("AttributeKeyNotFound", "Attribute does not contain key for provided time"));
+				Report(ELogVerbosity::Warning, LOCTEXT("AttributeKeyNotFound", "Attribute does not contain key for provided time"));
 			}	
 		}
 		else
@@ -2085,7 +2072,7 @@ bool UAnimDataController::RemoveAttributeKey(const FAnimationAttributeIdentifier
 	}
 	else
 	{
-		ReportError(LOCTEXT("InvalidAttributeIdentifier", "Invalid attribute identifier provided"));
+		Report(ELogVerbosity::Error, LOCTEXT("InvalidAttributeIdentifier", "Invalid attribute identifier provided"));
 	}
 
 	return false;
diff --git a/Engine/Source/Developer/AnimationDataController/Public/AnimDataController.h b/Engine/Source/Developer/AnimationDataController/Public/AnimDataController.h
--- a/Engine/Source/Developer/AnimationDataController/Public/AnimDataController.h
+++ b/Engine/Source/Developer/AnimationDataController/Public/AnimDataController.h
@@ -122,19 +122,24 @@ private:
 	bool CheckOuterClass(UClass* InClass) const;
 
 	/** Helper functionality to output script-based warnings and errors */
-	void ReportWarning(const FText& InMessage) const;
-	void ReportError(const FText& InMessage) const;
+	void Report(ELogVerbosity::Type InVerbosity, const FText& InMessage) const;
 
 	template <typename FmtType, typename... Types>
     void ReportWarningf(const FmtType& Fmt, Types... Args) const
 	{	
-		ReportWarning(FText::Format(Fmt, Args...));
+		Report(ELogVerbosity::Warning, FText::Format(Fmt, Args...));
 	}
 
 	template <typename FmtType, typename... Types>
     void ReportErrorf(const FmtType& Fmt, Types... Args) const
 	{
-		ReportError(FText::Format(Fmt, Args...));
+		Report(ELogVerbosity::Error, FText::Format(Fmt, Args...));
+	}
+
+	template <typename FmtType, typename... Types>
+	void Reportf(ELogVerbosity::Type InLogType, const FmtType& Fmt, Types... Args) const
+	{	
+		Report(InLogType, FText::Format(Fmt, Args...));
 	}
 #endif // WITH_EDITOR
 
