Reintroducing AnimationDataModel refactor with additional memory and cooktime optimizations.

diff --git a/Engine/Plugins/Interchange/Runtime/Source/Import/Private/Animation/InterchangeAnimSequenceFactory.cpp b/Engine/Plugins/Interchange/Runtime/Source/Import/Private/Animation/InterchangeAnimSequenceFactory.cpp
--- a/Engine/Plugins/Interchange/Runtime/Source/Import/Private/Animation/InterchangeAnimSequenceFactory.cpp
+++ b/Engine/Plugins/Interchange/Runtime/Source/Import/Private/Animation/InterchangeAnimSequenceFactory.cpp
@@ -197,7 +197,7 @@ namespace UE::Interchange::Private
 		FAnimationCurveIdentifier FloatCurveId(NewName, ERawCurveTrackTypes::RCT_Float);
 
 
-		UAnimDataModel* DataModel = TargetSequence->GetDataModel();
+		IAnimationDataModel* DataModel = TargetSequence->GetDataModel();
 		IAnimationDataController& Controller = TargetSequence->GetController();
 
 		const FFloatCurve* TargetCurve = DataModel->FindFloatCurve(FloatCurveId);
@@ -328,7 +328,11 @@ namespace UE::Interchange::Private
 
 			//This destroy all previously imported animation raw data
 			Controller.RemoveAllBoneTracks();
-			Controller.SetPlayLength(FGenericPlatformMath::Max<float>(SequenceLength, MINIMUM_ANIMATION_LENGTH));
+
+			const FFrameRate ResampleFrameRate(SampleRate, 1);
+			const FFrameNumber NumberOfFrames = ResampleFrameRate.AsFrameNumber(SequenceLength);
+			Controller.SetFrameRate(ResampleFrameRate);
+			Controller.SetNumberOfFrames(FGenericPlatformMath::Max<int32>(NumberOfFrames.Value, 1));
 
 			FTransform3f GlobalOffsetTransform;
 			{
@@ -454,7 +458,7 @@ namespace UE::Interchange::Private
 		AnimSequenceFactoryNode->GetCustomImportAttributeCurves(bImportAttributeCurves);
 		if (bImportAttributeCurves)
 		{
-			const UAnimDataModel* DataModel = AnimSequence->GetDataModel();
+			const IAnimationDataModel* DataModel = AnimSequence->GetDataModel();
 			const int32 NumFloatCurves = DataModel->GetNumberOfFloatCurves();
 			const FAnimationCurveData& CurveData = DataModel->GetCurveData();
 
@@ -905,6 +909,7 @@ UObject* UInterchangeAnimSequenceFactory::CreateAsset(const FCreateAssetParams&
 		IAnimationDataController& Controller = AnimSequence->GetController();
 		Controller.OpenBracket(NSLOCTEXT("InterchangeAnimSequenceFactory", "ImportAnimationInterchange_Bracket", "Importing Animation (Interchange)"));
 		AnimSequence->SetSkeleton(Skeleton);
+		Controller.InitializeModel();
 		AnimSequence->ImportFileFramerate = SampleRate;
 		AnimSequence->ImportResampleFramerate = SampleRate;
 		Controller.SetFrameRate(FrameRate);
diff --git a/Engine/Plugins/VirtualProduction/Takes/Source/TakeTrackRecorders/Private/MovieSceneAnimationTrackRecorder.cpp b/Engine/Plugins/VirtualProduction/Takes/Source/TakeTrackRecorders/Private/MovieSceneAnimationTrackRecorder.cpp
--- a/Engine/Plugins/VirtualProduction/Takes/Source/TakeTrackRecorders/Private/MovieSceneAnimationTrackRecorder.cpp
+++ b/Engine/Plugins/VirtualProduction/Takes/Source/TakeTrackRecorders/Private/MovieSceneAnimationTrackRecorder.cpp
@@ -393,12 +393,13 @@ bool UMovieSceneAnimationTrackRecorder::LoadRecordedFile(const FString& FileName
 								IAnimationDataController& Controller = AnimSequence->GetController();
 								{
 									IAnimationDataController::FScopedBracket ScopedBracket(Controller, LOCTEXT("LoadRecordedFile_Bracket", "Loading recorded animation file"));
-
+									Controller.InitializeModel();
 									Controller.ResetModel();
 
 									const float FloatDenominator = 1000.0f;
 									const float Numerator = FloatDenominator / Header.IntervalTime;
-									Controller.SetFrameRate(FFrameRate(Numerator, FloatDenominator));
+									const FFrameRate FrameRate(Numerator, FloatDenominator);
+									Controller.SetFrameRate(FrameRate);
 
 									int32 MaxNumberOfKeys = 0;
 									for (int32 TrackIndex = 0; TrackIndex < Header.AnimationTrackNames.Num(); ++TrackIndex)
@@ -425,7 +426,9 @@ bool UMovieSceneAnimationTrackRecorder::LoadRecordedFile(const FString& FileName
 										Controller.SetBoneTrackKeys(Header.AnimationTrackNames[TrackIndex], PosKeys, RotKeys, ScaleKeys);
 									}
 
-									Controller.SetPlayLength((MaxNumberOfKeys > 1) ? (MaxNumberOfKeys - 1) * Header.IntervalTime : MINIMUM_ANIMATION_LENGTH);
+									const FFrameNumber FrameNumber = FrameRate.AsFrameNumber((MaxNumberOfKeys > 1) ? (MaxNumberOfKeys - 1) * Header.IntervalTime : FrameRate.AsInterval());
+									Controller.SetNumberOfFrames(FrameNumber);
+
 									Controller.NotifyPopulated();
 								}
 
@@ -439,7 +442,7 @@ bool UMovieSceneAnimationTrackRecorder::LoadRecordedFile(const FString& FileName
 								FSavePackageArgs SaveArgs;
 								SaveArgs.TopLevelFlags = RF_Standalone;
 								SaveArgs.SaveFlags = SAVE_NoError;
-								UPackage::SavePackage(Package, NULL, *PackageFileName, SaveArgs);
+								UPackage::SavePackage(Package, nullptr, *PackageFileName, SaveArgs);
 
 
 								FFrameNumber SequenceLength = (AnimSequence->GetPlayLength() * TickResolution).FloorToFrame();
diff --git a/Engine/Source/Developer/AnimationDataController/Private/AnimDataController.cpp b/Engine/Source/Developer/AnimationDataController/Private/AnimDataController.cpp
--- a/Engine/Source/Developer/AnimationDataController/Private/AnimDataController.cpp
+++ b/Engine/Source/Developer/AnimationDataController/Private/AnimDataController.cpp
@@ -4,89 +4,42 @@
 #include "AnimDataControllerActions.h"
 #include "MovieSceneTimeHelpers.h"
 
-#include "Animation/AnimData/AnimDataModel.h"
+#include "Animation/AnimData/IAnimationDataModel.h"
 #include "Animation/AnimData/CurveIdentifier.h"
-
-#include "Animation/AnimSequenceBase.h"
 #include "Animation/AnimSequence.h"
 
-
 #include "Algo/Transform.h"
+#include "Animation/AnimationSettings.h"
 #include "UObject/NameTypes.h"
 #include "Animation/AnimCurveTypes.h"
+#include "Animation/AnimSequenceHelpers.h"
 #include "Math/UnrealMathUtility.h"
 
 #define LOCTEXT_NAMESPACE "AnimDataController"
 
 #if WITH_EDITOR
-
-namespace UE {
-namespace Anim {
-	bool CanTransactChanges()
-	{
-		return GEngine && GEngine->CanTransact() && !GIsTransacting;
-	}
-
-	struct FScopedCompoundTransaction
-	{
-		FScopedCompoundTransaction(UE::FChangeTransactor& InTransactor, const FText& InDescription) : Transactor(InTransactor), bCreated(false)
-		{
-			if (CanTransactChanges() && !Transactor.IsTransactionPending())
-			{
-				Transactor.OpenTransaction(InDescription);
-				bCreated = true;
-			}
-		}
-
-		~FScopedCompoundTransaction()
-		{
-			if (bCreated)
-			{
-				Transactor.CloseTransaction();
-			}
-		}
-
-		UE::FChangeTransactor& Transactor;
-		bool bCreated;
-	};
-}}
-
-#define CONDITIONAL_TRANSACTION(Text) \
-	TUniquePtr<UE::Anim::FScopedCompoundTransaction> Transaction; \
-	if (UE::Anim::CanTransactChanges() && bShouldTransact) \
-	{ \
-		Transaction = MakeUnique<UE::Anim::FScopedCompoundTransaction>(ChangeTransactor, Text); \
-	}
-
-#define CONDITIONAL_BRACKET(Text) IAnimationDataController::FScopedBracket Transaction(this, Text, UE::Anim::CanTransactChanges() && bShouldTransact);
-
-#define CONDITIONAL_ACTION(ActionClass, ...) \
-	if (UE::Anim::CanTransactChanges() && bShouldTransact) \
-	{ \
-		ChangeTransactor.AddTransactionChange<ActionClass>(__VA_ARGS__); \
-	}
-
-void UAnimDataController::SetModel(UAnimDataModel* InModel)
+void UAnimDataController::SetModel(TScriptInterface<IAnimationDataModel> InModel)
 {	
 	if (Model != nullptr)
 	{
 		Model->GetModifiedEvent().RemoveAll(this);
 	}
 
-	Model = InModel;
+	ModelInterface = InModel;
+	Model = CastChecked<UAnimDataModel>(InModel.GetObject(), ECastCheckedType::NullAllowed);
 	
-	ChangeTransactor.SetTransactionObject(InModel);
+	ChangeTransactor.SetTransactionObject(Model.Get());
 }
 
 void UAnimDataController::OpenBracket(const FText& InTitle, bool bShouldTransact /*= true*/)
 {
 	ValidateModel();
 
-	if (UE::Anim::CanTransactChanges() && !ChangeTransactor.IsTransactionPending())
+	if (UE::FChangeTransactor::CanTransactChanges() && !ChangeTransactor.IsTransactionPending())
 	{
 		ChangeTransactor.OpenTransaction(InTitle);
 
-		CONDITIONAL_ACTION(UE::Anim::FCloseBracketAction, InTitle.ToString());
+		ConditionalAction<UE::Anim::FCloseBracketAction>(bShouldTransact, InTitle.ToString());
 	}
 
 	if (BracketDepth == 0)
@@ -94,7 +47,7 @@ void UAnimDataController::OpenBracket(const FText& InTitle, bool bShouldTransact
 		FBracketPayload Payload;
 		Payload.Description = InTitle.ToString();
 
-		Model->Notify(EAnimDataModelNotifyType::BracketOpened, Payload);
+		Model->GetNotifier().Notify(EAnimDataModelNotifyType::BracketOpened, Payload);
 	}
 
 	++BracketDepth;
@@ -114,95 +67,181 @@ void UAnimDataController::CloseBracket(bool bShouldTransact /*= true*/)
 
 	if (BracketDepth == 0)
 	{
-		if (UE::Anim::CanTransactChanges())
+		if (UE::FChangeTransactor::CanTransactChanges())
 		{
 			ensure(ChangeTransactor.IsTransactionPending());
 
-			CONDITIONAL_ACTION(UE::Anim::FOpenBracketAction, TEXT("Open Bracket"));
+			ConditionalAction<UE::Anim::FOpenBracketAction>(bShouldTransact, TEXT("Open Bracket"));
 
 			ChangeTransactor.CloseTransaction();
 		}
 		
-		Model->Notify(EAnimDataModelNotifyType::BracketClosed);
-    }
+		Model->GetNotifier().Notify(EAnimDataModelNotifyType::BracketClosed);
+	}
 }
 
-void UAnimDataController::SetPlayLength(float Length, bool bShouldTransact /*= true*/)
+void UAnimDataController::SetNumberOfFrames(FFrameNumber Length, bool bShouldTransact)
 {
 	ValidateModel();
+	const FFrameNumber CurrentNumberOfFrames = Model->GetNumberOfFrames();
+
+	const int32 DeltaFrames = FMath::Abs(Length.Value - CurrentNumberOfFrames.Value);
+	
+	const FFrameNumber T0 = Length > CurrentNumberOfFrames ? CurrentNumberOfFrames : CurrentNumberOfFrames - DeltaFrames;
+	const FFrameNumber T1 = Length > CurrentNumberOfFrames ? Length : CurrentNumberOfFrames;
 
-	// Calculate whether or new play length is shorter or longer than current, set-up T0; T1 accordingly
-	// Assumption is made that time is always added or removed at/from the end
-	// Added: T0 = current length, T1 = new length
-	// Removed: T0 = current length - removed length, T1 = current length
-	const float Delta = FMath::Abs(Length - Model->PlayLength);
-	const float T0 = Length > Model->PlayLength ? Model->PlayLength : Model->PlayLength - Delta;
-	const float T1 = Length > Model->PlayLength ? Length : Model->PlayLength;
-	ResizePlayLength(Length, T0, T1, bShouldTransact);	
+	ResizeNumberOfFrames(Length, T0, T1, bShouldTransact);
 }
 
-void UAnimDataController::Resize(float Length, float T0, float T1, bool bShouldTransact /*= true*/)
+void UAnimDataController::ResizeNumberOfFrames(FFrameNumber NewLength, FFrameNumber T0, FFrameNumber T1, bool bShouldTransact)
 {
 	ValidateModel();
 	
-	const TRange<float> PlayRange(TRange<float>::BoundsType::Inclusive(0.f), TRange<float>::BoundsType::Inclusive(Model->PlayLength));
-	if (!FMath::IsNearlyZero(Length) && Length > 0.f)
+	const TRange<FFrameNumber> PlayRange(TRange<FFrameNumber>::BoundsType::Inclusive(0), TRange<FFrameNumber>::BoundsType::Exclusive(FMath::Max(1, Model->GetNumberOfKeys())));
+	if (NewLength >= 0)
 	{
-		if (Length != Model->PlayLength)
+		if (NewLength != Model->GetNumberOfFrames())
 		{
-			// Ensure that T0 is within the curent play range
+			// Ensure that T0 is within the current play range
 			if (PlayRange.Contains(T0))
 			{
 				// Ensure that the start and end length of either removal or insertion are valid
 				if (T0 < T1)
 				{
-					CONDITIONAL_BRACKET(LOCTEXT("ResizeModel", "Resizing Animation Data"));
-					const bool bInserted = Length > Model->PlayLength;
-					ResizePlayLength(Length, T0, T1, bShouldTransact);
-					ResizeCurves(Length, bInserted, T0, T1, bShouldTransact);
-					ResizeAttributes(Length, bInserted, T0, T1, bShouldTransact);
+					FTransaction Transaction = ConditionalTransaction(LOCTEXT("ResizePlayLength", "Resizing Play Length"), bShouldTransact);
+
+					const FFrameRate CurrentFrameRate = Model->GetFrameRate();
+					const FFrameNumber CurrentNumberOfFrames = Model->GetNumberOfFrames();
+
+					FSequenceLengthChangedPayload Payload;
+					PRAGMA_DISABLE_DEPRECATION_WARNINGS
+					Payload.T0 = CurrentFrameRate.AsSeconds(T0);
+					Payload.T1 = CurrentFrameRate.AsSeconds(T1);
+					Payload.PreviousLength = CurrentFrameRate.AsSeconds(CurrentNumberOfFrames);
+					PRAGMA_ENABLE_DEPRECATION_WARNINGS
+
+					Payload.PreviousNumberOfFrames = CurrentNumberOfFrames;
+					Payload.Frame0 = T0;
+					Payload.Frame1 = T1;
+
+					ConditionalAction<UE::Anim::FResizePlayLengthInFramesAction>(bShouldTransact, Model.Get(), Payload.Frame0, Payload.Frame1);
+
+					Model->NumberOfFrames = NewLength.Value;
+					Model->NumberOfKeys = Model->NumberOfFrames + 1;
+	
+					Model->GetNotifier().Notify<FSequenceLengthChangedPayload>(EAnimDataModelNotifyType::SequenceLengthChanged, Payload);
 				}
 				else
 				{
-					ReportErrorf(LOCTEXT("InvalidEndTimeError", "Invalid T1, smaller that T0 value: T0 {0}, T1 {1}"), FText::AsNumber(T0), FText::AsNumber(T1));
+					ReportErrorf(LOCTEXT("InvalidEndTimeError", "Invalid T1, smaller that T0 value: T0 {0}, T1 {1}"), FText::AsNumber(T0.Value), FText::AsNumber(T0.Value));
 				}
 			}
 			else
 			{
-				ReportErrorf(LOCTEXT("InvalidStartTimeError", "Invalid T0, not within existing play range: T0 {0}, Play Length {1}"), FText::AsNumber(T0), FText::AsNumber(Model->PlayLength));
-			}			
+				ReportErrorf(LOCTEXT("InvalidStartTimeError", "Invalid T0, not within existing play range: T0 {0}, Play Length {1}"), FText::AsNumber(T0.Value), FText::AsNumber(Model->GetPlayLength()));
+			}
 		}
-		else
+		else if (Model->bPopulated)
 		{
-			ReportWarningf(LOCTEXT("SamePlayLengthWarning", "New play length is same as existing one: {0} seconds"), FText::AsNumber(Length));
+			ReportWarningf(LOCTEXT("SamePlayLengthWarning", "New play length is same as existing one: {0} frames"), FText::AsNumber(NewLength.Value));
 		}
 	}
 	else
 	{
-		ReportErrorf(LOCTEXT("InvalidPlayLengthError", "Invalid play length value provided: {0} seconds"), FText::AsNumber(Length));
+		ReportErrorf(LOCTEXT("InvalidPlayLengthError", "Invalid play length value provided: {0} frames"), FText::AsNumber(NewLength.Value));
 	}
 }
 
+void UAnimDataController::ResizeInFrames(FFrameNumber NewLength, FFrameNumber T0, FFrameNumber T1, bool bShouldTransact)
+{
+	ValidateModel();
+	
+	const int32 CurrentNumberOFrames = Model->GetNumberOfFrames();
+	
+	const TRange<FFrameNumber> PlayRange(TRange<FFrameNumber>::BoundsType::Inclusive(0), TRange<FFrameNumber>::BoundsType::Exclusive(FMath::Max(1,Model->GetNumberOfKeys())));
+	if (NewLength >= 0)
+	{
+		if (NewLength != Model->GetNumberOfFrames())
+		{
+			// Ensure that T0 is within the current play range
+			if (PlayRange.Contains(T0))
+			{
+				// Ensure that the start and end length of either removal or insertion are valid
+				if (T0 < T1)
+				{
+					FBracket Bracket = ConditionalBracket(LOCTEXT("ResizeModel", "Resizing Animation Data"), bShouldTransact);
+
+					const bool bInserted = NewLength > CurrentNumberOFrames;
+					ResizeNumberOfFrames(NewLength, T0, T1, bShouldTransact);
+					
+					const FFrameRate CurrentFrameRate = Model->GetFrameRate();
+					ResizeCurves(CurrentFrameRate.AsSeconds(NewLength), bInserted, CurrentFrameRate.AsSeconds(T0), CurrentFrameRate.AsSeconds(T1), bShouldTransact);
+					ResizeAttributes(CurrentFrameRate.AsSeconds(NewLength), bInserted, CurrentFrameRate.AsSeconds(T0), CurrentFrameRate.AsSeconds(T1), bShouldTransact);
+				}
+				else
+                {
+                	ReportErrorf(LOCTEXT("InvalidEndTimeError", "Invalid T1, smaller that T0 value: T0 {0}, T1 {1}"), FText::AsNumber(T0.Value), FText::AsNumber(T1.Value));
+                }
+            }
+            else
+            {
+                ReportErrorf(LOCTEXT("InvalidStartTimeError", "Invalid T0, not within existing play range: T0 {0}, Play Length {1}"), FText::AsNumber(T0.Value), FText::AsNumber(CurrentNumberOFrames));
+            }			
+        }
+		else if (Model->bPopulated)
+        {
+            ReportWarningf(LOCTEXT("SameGetPlayLengthWarning", "New play length is same as existing one: {0} frames"), FText::AsNumber(CurrentNumberOFrames));
+        }
+    }
+    else
+    {
+        ReportErrorf(LOCTEXT("InvalidGetPlayLengthError", "Invalid play length value provided: {0} frames"), FText::AsNumber(CurrentNumberOFrames));
+    }
+}
+
+void UAnimDataController::SetPlayLength(float Length, bool bShouldTransact /*= true*/)
+{
+	SetNumberOfFrames(ConvertSecondsToFrameNumber(Length), bShouldTransact);
+}
+
+void UAnimDataController::Resize(float Length, float T0, float T1, bool bShouldTransact /*= true*/)
+{
+	ValidateModel();
+	
+	ResizeInFrames(ConvertSecondsToFrameNumber(Length), ConvertSecondsToFrameNumber(T0), ConvertSecondsToFrameNumber(T1), bShouldTransact);
+}
+
 void UAnimDataController::SetFrameRate(FFrameRate FrameRate, bool bShouldTransact /*= true*/)
 {
 	ValidateModel();
 
 	// Disallow invalid frame-rates, or 0.0 intervals
-	const float FrameRateInterval = FrameRate.AsInterval();
-	if ( FrameRate.IsValid() && !FMath::IsNearlyZero(FrameRateInterval) && FrameRateInterval > 0.f)
+	const double FrameRateInterval = FrameRate.AsInterval();
+	if ( FrameRate.IsValid() && !FMath::IsNearlyZero(FrameRateInterval) && FrameRateInterval > 0.0)
 	{
-		CONDITIONAL_TRANSACTION(LOCTEXT("SetFrameRate", "Setting Frame Rate"));
-
-		CONDITIONAL_ACTION(UE::Anim::FSetFrameRateAction, Model.Get());
+		// Need to verify framerate
+		const FFrameRate CurrentFrameRate = Model->GetFrameRate();
+		if (FrameRate.IsMultipleOf(CurrentFrameRate) || FrameRate.IsFactorOf(CurrentFrameRate) || !Model->bPopulated)
+		{
+			FTransaction Transaction = ConditionalTransaction(LOCTEXT("SetFrameRate", "Setting Frame Rate"), bShouldTransact);
+			ConditionalAction<UE::Anim::FSetFrameRateAction>(bShouldTransact, Model.Get());
 
-		FFrameRateChangedPayload Payload;
-		Payload.PreviousFrameRate = Model->FrameRate;
+			const FFrameNumber CurrentNumberOfFrames = Model->GetNumberOfFrames();
+			const FFrameTime ConvertedLastFrameTime = FFrameRate::TransformTime(CurrentNumberOfFrames, CurrentFrameRate, FrameRate);
+			ensure(FMath::IsNearlyZero(ConvertedLastFrameTime.GetSubFrame()) || !Model->bPopulated);
 			
 			Model->FrameRate = FrameRate;
-		Model->NumberOfFrames = Model->FrameRate.AsFrameTime(Model->PlayLength).RoundToFrame().Value;
+			Model->NumberOfFrames = ConvertedLastFrameTime.GetFrame().Value;
 			Model->NumberOfKeys = Model->NumberOfFrames + 1;
 			
-		Model->Notify(EAnimDataModelNotifyType::FrameRateChanged, Payload);
+			FFrameRateChangedPayload Payload;
+			Payload.PreviousFrameRate = CurrentFrameRate;
+			Model->GetNotifier().Notify(EAnimDataModelNotifyType::FrameRateChanged, Payload);
+		}
+		else
+        {
+        	ReportErrorf(LOCTEXT("NonCompatibleFrameRateError", "Incompatible frame rate provided: {0} not a multiple or fact or {1}"), FrameRate.ToPrettyText(), CurrentFrameRate.ToPrettyText());
+        }
 	}
 	else
 	{
@@ -219,7 +258,7 @@ void UAnimDataController::UpdateCurveNamesFromSkeleton(const USkeleton* Skeleton
 	{
 		if (IsSupportedCurveType(SupportedCurveType))
 		{
-			CONDITIONAL_BRACKET(LOCTEXT("ValidateRawCurves", "Validating Animation Curve Names"));
+			FBracket Bracket = ConditionalBracket(LOCTEXT("ValidateRawCurves", "Validating Animation Curve Names"), bShouldTransact);
 			switch (SupportedCurveType)
 			{
 			case ERawCurveTrackTypes::RCT_Float:
@@ -254,12 +293,17 @@ void UAnimDataController::UpdateCurveNamesFromSkeleton(const USkeleton* Skeleton
 				}
 				break;
 			}
+			case ERawCurveTrackTypes::RCT_Vector:
+			default:
+				const FString CurveTypeAsString = GetCurveTypeValueName(SupportedCurveType);
+				ReportWarningf(LOCTEXT("InvalidCurveTypeWarning", "Invalid curve type provided: {0} ({1})"), FText::FromString(CurveTypeAsString), FText::AsNumber(static_cast<int32>(SupportedCurveType)));
+				
 			}
 		}
 		else
 		{
 			const FString CurveTypeAsString = GetCurveTypeValueName(SupportedCurveType);
-			ReportWarningf(LOCTEXT("InvalidCurveTypeWarning", "Invalid curve type provided: {0} ({1})"), FText::FromString(CurveTypeAsString), FText::AsNumber((int32)SupportedCurveType));
+			ReportWarningf(LOCTEXT("InvalidCurveTypeWarning", "Invalid curve type provided: {0} ({1})"), FText::FromString(CurveTypeAsString), FText::AsNumber(static_cast<int32>(SupportedCurveType)));
 		}
 	}
 	else
@@ -276,7 +320,7 @@ void UAnimDataController::FindOrAddCurveNamesOnSkeleton(USkeleton* Skeleton, ERa
 	{
 		if (IsSupportedCurveType(SupportedCurveType))
 		{
-			CONDITIONAL_BRACKET(LOCTEXT("FindOrAddRawCurveNames", "Updating Skeleton with Animation Curve Names"));
+			FBracket Bracket = ConditionalBracket(LOCTEXT("FindOrAddRawCurveNames", "Updating Skeleton with Animation Curve Names"), bShouldTransact);
 			switch (SupportedCurveType)
 			{
 			case ERawCurveTrackTypes::RCT_Float:
@@ -309,12 +353,16 @@ void UAnimDataController::FindOrAddCurveNamesOnSkeleton(USkeleton* Skeleton, ERa
 				}
 				break;
 			}
+			case ERawCurveTrackTypes::RCT_Vector:
+			default:
+				const FString CurveTypeAsString = GetCurveTypeValueName(SupportedCurveType);
+				ReportWarningf(LOCTEXT("InvalidCurveTypeWarning", "Invalid curve type provided: {0} ({1})"), FText::FromString(CurveTypeAsString), FText::AsNumber(static_cast<int32>(SupportedCurveType)));
 			}
 		}
 		else
 		{
 			const FString CurveTypeAsString = GetCurveTypeValueName(SupportedCurveType);
-			ReportWarningf(LOCTEXT("InvalidCurveTypeWarning", "Invalid curve type provided: {0} ({1})"), FText::FromString(CurveTypeAsString), FText::AsNumber((int32)SupportedCurveType));
+			ReportWarningf(LOCTEXT("InvalidCurveTypeWarning", "Invalid curve type provided: {0} ({1})"), FText::FromString(CurveTypeAsString), FText::AsNumber(static_cast<int32>(SupportedCurveType)));
 		}
 	}
 	else
@@ -325,7 +373,7 @@ void UAnimDataController::FindOrAddCurveNamesOnSkeleton(USkeleton* Skeleton, ERa
 
 bool UAnimDataController::RemoveBoneTracksMissingFromSkeleton(const USkeleton* Skeleton, bool bShouldTransact /*= true*/)
 {
-	if (!CheckOuterClass(UAnimSequence::StaticClass()))
+	if (!ModelInterface->GetAnimationSequence())
 	{
 		return false;
 	}
@@ -357,7 +405,7 @@ bool UAnimDataController::RemoveBoneTracksMissingFromSkeleton(const USkeleton* S
 
 		if (TracksToBeRemoved.Num() || TracksUpdated.Num())
 		{
-			CONDITIONAL_BRACKET(LOCTEXT("RemoveBoneTracksMissingFromSkeleton", "Validating Bone Animation Track Data against Skeleton"));
+			FBracket Bracket = ConditionalBracket(LOCTEXT("RemoveBoneTracksMissingFromSkeleton", "Validating Bone Animation Track Data against Skeleton"), bShouldTransact);
 			for (const FName& TrackName : TracksToBeRemoved)
 			{
 				RemoveBoneTrack(TrackName);
@@ -367,7 +415,7 @@ bool UAnimDataController::RemoveBoneTracksMissingFromSkeleton(const USkeleton* S
 			{
 				FAnimationTrackChangedPayload Payload;
 				Payload.Name = TrackName;
-				Model->Notify(EAnimDataModelNotifyType::TrackChanged, Payload);
+				Model->GetNotifier().Notify(EAnimDataModelNotifyType::TrackChanged, Payload);
 			}
 		}
 
@@ -407,7 +455,7 @@ void UAnimDataController::UpdateAttributesFromSkeleton(const USkeleton* Skeleton
 
 		if (ToRemoveIdentifiers.Num() || ToDuplicateIdentifiers.Num())
 		{
-			CONDITIONAL_BRACKET(LOCTEXT("VerifyAttributeBoneNames", "Remapping Animation Attribute Data"));
+			FBracket Bracket = ConditionalBracket(LOCTEXT("VerifyAttributeBoneNames", "Remapping Animation Attribute Data"), bShouldTransact);
 			for (const FAnimationAttributeIdentifier& Identifier : ToRemoveIdentifiers)
 			{
 				RemoveAttribute(Identifier);
@@ -415,8 +463,8 @@ void UAnimDataController::UpdateAttributesFromSkeleton(const USkeleton* Skeleton
 			
 			for (const TPair<FAnimationAttributeIdentifier, int32>& Pair : ToDuplicateIdentifiers)
 			{
-				FAnimationAttributeIdentifier NewIdentifier = Pair.Key;
-				NewIdentifier.BoneIndex = Pair.Value;
+				const FAnimationAttributeIdentifier& DuplicateIdentifier = Pair.Key;
+				FAnimationAttributeIdentifier NewIdentifier(DuplicateIdentifier.GetName(), Pair.Value, DuplicateIdentifier.GetBoneName(), DuplicateIdentifier.GetType());
 
 				DuplicateAttribute(Pair.Key, NewIdentifier);
 				RemoveAttribute(Pair.Key);
@@ -433,17 +481,22 @@ void UAnimDataController::ResetModel(bool bShouldTransact /*= true*/)
 {
 	ValidateModel();
 
-	CONDITIONAL_BRACKET(LOCTEXT("ResetModel", "Clearing Animation Data"));
+	FBracket Bracket = ConditionalBracket(LOCTEXT("ResetModel", "Clearing Animation Data"), bShouldTransact);
 
-	RemoveAllBoneTracks(bShouldTransact);
+	if (Model->GetOuter()->IsA<UAnimSequence>())
+	{
+		RemoveAllBoneTracks(bShouldTransact);
+	}
 
 	RemoveAllCurvesOfType(ERawCurveTrackTypes::RCT_Float, bShouldTransact);
 	RemoveAllCurvesOfType(ERawCurveTrackTypes::RCT_Transform, bShouldTransact);
+	RemoveAllAttributes(bShouldTransact);
 
-	SetPlayLength(MINIMUM_ANIMATION_LENGTH, bShouldTransact);
-	SetFrameRate(FFrameRate(30,1), bShouldTransact);
 
-	Model->Notify(EAnimDataModelNotifyType::Reset);
+	SetFrameRate(UAnimationSettings::Get()->GetDefaultFrameRate(), bShouldTransact);
+	SetNumberOfFrames(1, bShouldTransact);
+
+	Model->GetNotifier().Notify(EAnimDataModelNotifyType::Reset);
 }
 
 bool UAnimDataController::AddCurve(const FAnimationCurveIdentifier& CurveId, int32 CurveFlags /*= EAnimAssetCurveFlags::AACF_Editable*/, bool bShouldTransact /*= true*/)
@@ -455,7 +508,7 @@ bool UAnimDataController::AddCurve(const FAnimationCurveIdentifier& CurveId, int
 		{
 			if (!Model->FindCurve(CurveId))
 			{
-				CONDITIONAL_TRANSACTION(LOCTEXT("AddRawCurve", "Adding Animation Curve"));
+				FTransaction Transaction = ConditionalTransaction(LOCTEXT("AddRawCurve", "Adding Animation Curve"), bShouldTransact);
 
 				FCurveAddedPayload Payload;
 				Payload.Identifier = CurveId;
@@ -473,23 +526,27 @@ bool UAnimDataController::AddCurve(const FAnimationCurveIdentifier& CurveId, int
 				case ERawCurveTrackTypes::RCT_Float:
 					AddNewCurve(Model->CurveData.FloatCurves);
 					break;
+				case ERawCurveTrackTypes::RCT_Vector:
+				default:
+					const FString CurveTypeAsString = GetCurveTypeValueName(CurveId.CurveType);
+					ReportWarningf(LOCTEXT("InvalidCurveTypeWarning", "Invalid curve type provided: {0} ({1})"), FText::FromString(CurveTypeAsString), FText::AsNumber(static_cast<int32>(CurveId.CurveType)));
 				}
 
-				CONDITIONAL_ACTION(UE::Anim::FRemoveCurveAction, CurveId);
-				Model->Notify(EAnimDataModelNotifyType::CurveAdded, Payload);
+				ConditionalAction<UE::Anim::FRemoveCurveAction>(bShouldTransact, CurveId);
+				Model->GetNotifier().Notify(EAnimDataModelNotifyType::CurveAdded, Payload);
 
 				return true;
 			}
 			else
 			{
 				const FString CurveTypeAsString = GetCurveTypeValueName(CurveId.CurveType);
-				ReportWarningf(LOCTEXT("ExistingCurveNameWarning", "Curve with name {0} and type {1} ({2}) already exists"), FText::FromName(CurveId.InternalName.DisplayName), FText::FromString(CurveTypeAsString), FText::AsNumber((int32)CurveId.CurveType));
+				ReportWarningf(LOCTEXT("ExistingCurveNameWarning", "Curve with name {0} and type {1} ({2}) already exists"), FText::FromName(CurveId.InternalName.DisplayName), FText::FromString(CurveTypeAsString), FText::AsNumber(static_cast<int32>(CurveId.CurveType)));
 			}			
 		}
 		else 
 		{
 			const FString CurveTypeAsString = GetCurveTypeValueName(CurveId.CurveType);
-			ReportWarningf(LOCTEXT("InvalidCurveTypeWarning", "Invalid curve type provided: {0} ({1})"), FText::FromString(CurveTypeAsString), FText::AsNumber((int32)CurveId.CurveType));
+			ReportWarningf(LOCTEXT("InvalidCurveTypeWarning", "Invalid curve type provided: {0} ({1})"), FText::FromString(CurveTypeAsString), FText::AsNumber(static_cast<int32>(CurveId.CurveType)));
 		}		
 	}
 	else
@@ -517,7 +574,7 @@ bool UAnimDataController::DuplicateCurve(const FAnimationCurveIdentifier& CopyCu
 				{
 					if (!Model->FindCurve(NewCurveId))
 					{
-						CONDITIONAL_TRANSACTION(LOCTEXT("CopyRawCurve", "Duplicating Animation Curve"));
+						FTransaction Transaction = ConditionalTransaction(LOCTEXT("CopyRawCurve", "Duplicating Animation Curve"), bShouldTransact);
 
 						auto DuplicateCurve = [NewCurveName = NewCurveId.InternalName](auto& CurveDataArray, const auto& SourceCurve)
 						{
@@ -533,33 +590,37 @@ bool UAnimDataController::DuplicateCurve(const FAnimationCurveIdentifier& CopyCu
 						case ERawCurveTrackTypes::RCT_Float:
 							DuplicateCurve(Model->CurveData.FloatCurves, Model->GetFloatCurve(CopyCurveId));
 							break;
+						case ERawCurveTrackTypes::RCT_Vector:
+						default:
+							const FString CurveTypeAsString = GetCurveTypeValueName(SupportedCurveType);
+							ReportWarningf(LOCTEXT("InvalidCurveTypeWarning", "Invalid curve type provided: {0} ({1})"), FText::FromString(CurveTypeAsString), FText::AsNumber(static_cast<int32>(SupportedCurveType)));
 						}
 
 						FCurveAddedPayload Payload;
 						Payload.Identifier = NewCurveId;
-						Model->Notify(EAnimDataModelNotifyType::CurveAdded, Payload);
+						Model->GetNotifier().Notify(EAnimDataModelNotifyType::CurveAdded, Payload);
 
-						CONDITIONAL_ACTION(UE::Anim::FRemoveCurveAction, NewCurveId);
+						ConditionalAction<UE::Anim::FRemoveCurveAction>(bShouldTransact, NewCurveId);
 
 						return true;
 					}
 					else
 					{
 						const FString CurveTypeAsString = GetCurveTypeValueName(NewCurveId.CurveType);
-						ReportWarningf(LOCTEXT("ExistingCurveNameWarning", "Curve with name {0} and type {1} ({2}) already exists"), FText::FromName(NewCurveId.InternalName.DisplayName), FText::FromString(CurveTypeAsString), FText::AsNumber((int32)NewCurveId.CurveType));
+						ReportWarningf(LOCTEXT("ExistingCurveNameWarning", "Curve with name {0} and type {1} ({2}) already exists"), FText::FromName(NewCurveId.InternalName.DisplayName), FText::FromString(CurveTypeAsString), FText::AsNumber(static_cast<int32>(NewCurveId.CurveType)));
 					}
 				}
 				else
 				{
 					const FString CurveTypeAsString = GetCurveTypeValueName(CopyCurveId.CurveType);
-					ReportWarningf(LOCTEXT("CurveNameToDuplicateNotFoundWarning", "Could not find curve with name {0} and type {1} ({2}) for duplication"), FText::FromName(NewCurveId.InternalName.DisplayName), FText::FromString(CurveTypeAsString), FText::AsNumber((int32)NewCurveId.CurveType));
+					ReportWarningf(LOCTEXT("CurveNameToDuplicateNotFoundWarning", "Could not find curve with name {0} and type {1} ({2}) for duplication"), FText::FromName(NewCurveId.InternalName.DisplayName), FText::FromString(CurveTypeAsString), FText::AsNumber(static_cast<int32>(NewCurveId.CurveType)));
 				}
 			}
 		}
 		else
 		{
 			const FString CurveTypeAsString = GetCurveTypeValueName(SupportedCurveType);
-			ReportWarningf(LOCTEXT("InvalidCurveTypeWarning", "Invalid curve type provided: {0} ({1})"), FText::FromString(CurveTypeAsString), FText::AsNumber((int32)SupportedCurveType));
+			ReportWarningf(LOCTEXT("InvalidCurveTypeWarning", "Invalid curve type provided: {0} ({1})"), FText::FromString(CurveTypeAsString), FText::AsNumber(static_cast<int32>(SupportedCurveType)));
 		}
 	}
 
@@ -576,32 +637,35 @@ bool UAnimDataController::RemoveCurve(const FAnimationCurveIdentifier& CurveId,
 	{
 		if (IsSupportedCurveType(CurveId.CurveType))
 		{
-			const FAnimCurveBase* Curve = Model->FindCurve(CurveId);
-			if (Curve)
+			if (Model->FindCurve(CurveId) != nullptr)
 			{
-				CONDITIONAL_TRANSACTION(LOCTEXT("RemoveCurve", "Removing Animation Curve"));
+				FTransaction Transaction = ConditionalTransaction(LOCTEXT("RemoveCurve", "Removing Animation Curve"), bShouldTransact);
 
 				switch (SupportedCurveType)
 				{
 					case ERawCurveTrackTypes::RCT_Transform:
 					{
 						const FTransformCurve& TransformCurve = Model->GetTransformCurve(CurveId);
-						CONDITIONAL_ACTION(UE::Anim::FAddTransformCurveAction, CurveId, TransformCurve.GetCurveTypeFlags(), TransformCurve);
+						ConditionalAction<UE::Anim::FAddTransformCurveAction>(bShouldTransact, CurveId, TransformCurve.GetCurveTypeFlags(), TransformCurve);
 						Model->CurveData.TransformCurves.RemoveAll([Name = TransformCurve.Name](const FTransformCurve& ToRemoveCurve) { return ToRemoveCurve.Name == Name; });
 						break;
 					}
 					case ERawCurveTrackTypes::RCT_Float:
 					{
 						const FFloatCurve& FloatCurve = Model->GetFloatCurve(CurveId);
-						CONDITIONAL_ACTION(UE::Anim::FAddFloatCurveAction, CurveId, FloatCurve.GetCurveTypeFlags(), FloatCurve.FloatCurve.GetConstRefOfKeys(), FloatCurve.Color);
+						ConditionalAction<UE::Anim::FAddFloatCurveAction>(bShouldTransact, CurveId, FloatCurve.GetCurveTypeFlags(), FloatCurve.FloatCurve.GetConstRefOfKeys(), FloatCurve.Color);
 						Model->CurveData.FloatCurves.RemoveAll([Name = FloatCurve.Name](const FFloatCurve& ToRemoveCurve) { return ToRemoveCurve.Name == Name; });
 						break;
 					}
+					case ERawCurveTrackTypes::RCT_Vector:
+					default:
+						const FString CurveTypeAsString = GetCurveTypeValueName(SupportedCurveType);
+						ReportWarningf(LOCTEXT("InvalidCurveTypeWarning", "Invalid curve type provided: {0} ({1})"), FText::FromString(CurveTypeAsString), FText::AsNumber(static_cast<int32>(SupportedCurveType)));
 				}
 
 				FCurveRemovedPayload Payload;
 				Payload.Identifier = CurveId;
-				Model->Notify(EAnimDataModelNotifyType::CurveRemoved, Payload);
+				Model->GetNotifier().Notify(EAnimDataModelNotifyType::CurveRemoved, Payload);
 
 				return true;
 			}
@@ -614,7 +678,7 @@ bool UAnimDataController::RemoveCurve(const FAnimationCurveIdentifier& CurveId,
 		else
 		{
 			const FString CurveTypeAsString = GetCurveTypeValueName(CurveId.CurveType);
-			ReportWarningf(LOCTEXT("InvalidCurveTypeWarning", "Invalid curve type provided: {0} ({1})"), FText::FromString(CurveTypeAsString), FText::AsNumber((int32)CurveId.CurveType));
+			ReportWarningf(LOCTEXT("InvalidCurveTypeWarning", "Invalid curve type provided: {0} ({1})"), FText::FromString(CurveTypeAsString), FText::AsNumber(static_cast<int32>(CurveId.CurveType)));
 		}
 	}
 
@@ -625,7 +689,7 @@ void UAnimDataController::RemoveAllCurvesOfType(ERawCurveTrackTypes SupportedCur
 {
 	ValidateModel();
 
-	CONDITIONAL_BRACKET(LOCTEXT("DeleteAllRawCurve", "Deleting All Animation Curve"));
+	FBracket Bracket = ConditionalBracket(LOCTEXT("DeleteAllRawCurve", "Deleting All Animation Curve"), bShouldTransact);
 	switch (SupportedCurveType)
 	{
 	case ERawCurveTrackTypes::RCT_Transform:
@@ -649,7 +713,7 @@ void UAnimDataController::RemoveAllCurvesOfType(ERawCurveTrackTypes SupportedCur
 	case ERawCurveTrackTypes::RCT_Vector:
 	default:
 		const FString CurveTypeAsString = GetCurveTypeValueName(SupportedCurveType);
-		ReportWarningf(LOCTEXT("InvalidCurveTypeWarning", "Invalid curve type provided: {0} ({1})"), FText::FromString(CurveTypeAsString), FText::AsNumber((int32)SupportedCurveType));
+		ReportWarningf(LOCTEXT("InvalidCurveTypeWarning", "Invalid curve type provided: {0} ({1})"), FText::FromString(CurveTypeAsString), FText::AsNumber(static_cast<int32>(SupportedCurveType)));
 	}
 
 }
@@ -658,7 +722,7 @@ bool UAnimDataController::SetCurveFlag(const FAnimationCurveIdentifier& CurveId,
 {
 	ValidateModel();
 
-	ERawCurveTrackTypes SupportedCurveType = CurveId.CurveType;
+	const ERawCurveTrackTypes SupportedCurveType = CurveId.CurveType;
 
 	FAnimCurveBase* Curve = nullptr;
 
@@ -670,14 +734,19 @@ bool UAnimDataController::SetCurveFlag(const FAnimationCurveIdentifier& CurveId,
 	{
 		Curve = Model->FindMutableTransformCurveById(CurveId);
 	}
+	else
+	{
+		const FString CurveTypeAsString = GetCurveTypeValueName(SupportedCurveType);
+		ReportWarningf(LOCTEXT("InvalidCurveTypeWarning", "Invalid curve type provided: {0} ({1})"), FText::FromString(CurveTypeAsString), FText::AsNumber(static_cast<int32>(SupportedCurveType)));
+	}
 	
 	if (Curve)
 	{
-		CONDITIONAL_TRANSACTION(LOCTEXT("SetCurveFlag", "Setting Raw Curve Flag"));
+		FTransaction Transaction = ConditionalTransaction(LOCTEXT("SetCurveFlag", "Setting Raw Curve Flag"), bShouldTransact);
 
 		const int32 CurrentFlags = Curve->GetCurveTypeFlags();
 
-		CONDITIONAL_ACTION(UE::Anim::FSetCurveFlagsAction, CurveId, CurrentFlags, SupportedCurveType);
+		ConditionalAction<UE::Anim::FSetCurveFlagsAction>(bShouldTransact, CurveId, CurrentFlags, SupportedCurveType);
 
 		FCurveFlagsChangedPayload Payload;
 		Payload.Identifier = CurveId;
@@ -685,7 +754,7 @@ bool UAnimDataController::SetCurveFlag(const FAnimationCurveIdentifier& CurveId,
 
 		Curve->SetCurveTypeFlag(Flag, bState);
 
-		Model->Notify(EAnimDataModelNotifyType::CurveFlagsChanged, Payload);
+		Model->GetNotifier().Notify(EAnimDataModelNotifyType::CurveFlagsChanged, Payload);
 
 		return true;
 	}
@@ -704,8 +773,7 @@ bool UAnimDataController::SetCurveFlags(const FAnimationCurveIdentifier& CurveId
 
 	FAnimCurveBase* Curve = nullptr;
 
-	ERawCurveTrackTypes SupportedCurveType = CurveId.CurveType;
-
+	const ERawCurveTrackTypes SupportedCurveType = CurveId.CurveType;
 	if (SupportedCurveType == ERawCurveTrackTypes::RCT_Float)
 	{
 		Curve = Model->FindMutableFloatCurveById(CurveId);
@@ -714,14 +782,19 @@ bool UAnimDataController::SetCurveFlags(const FAnimationCurveIdentifier& CurveId
 	{
 		Curve = Model->FindMutableTransformCurveById(CurveId);
 	}
+	else
+	{
+		const FString CurveTypeAsString = GetCurveTypeValueName(SupportedCurveType);
+		ReportWarningf(LOCTEXT("InvalidCurveTypeWarning", "Invalid curve type provided: {0} ({1})"), FText::FromString(CurveTypeAsString), FText::AsNumber(static_cast<int32>(SupportedCurveType)));
+	}
 
 	if (Curve)
 	{
-		CONDITIONAL_TRANSACTION(LOCTEXT("SetRawCurveFlag", "Setting Raw Curve Flags"));
+		FTransaction Transaction = ConditionalTransaction(LOCTEXT("SetRawCurveFlag", "Setting Raw Curve Flags"), bShouldTransact);
 
 		const int32 CurrentFlags = Curve->GetCurveTypeFlags();
 
-		CONDITIONAL_ACTION(UE::Anim::FSetCurveFlagsAction, CurveId, CurrentFlags, SupportedCurveType);
+		ConditionalAction<UE::Anim::FSetCurveFlagsAction>(bShouldTransact, CurveId, CurrentFlags, SupportedCurveType);
 
 		FCurveFlagsChangedPayload Payload;
 		Payload.Identifier = CurveId;
@@ -729,7 +802,7 @@ bool UAnimDataController::SetCurveFlags(const FAnimationCurveIdentifier& CurveId
 
 		Curve->SetCurveTypeFlags(Flags);
 
-		Model->Notify(EAnimDataModelNotifyType::CurveFlagsChanged, Payload);
+		Model->GetNotifier().Notify(EAnimDataModelNotifyType::CurveFlagsChanged, Payload);
 
 		return true;
 	}
@@ -748,11 +821,9 @@ bool UAnimDataController::SetTransformCurveKeys(const FAnimationCurveIdentifier&
 
 	if (TransformValues.Num() == TimeKeys.Num())
 	{
-		FTransformCurve* Curve = Model->FindMutableTransformCurveById(CurveId);
-
-		if (Curve)
+		if (Model->FindMutableTransformCurveById(CurveId) != nullptr)
 		{
-			CONDITIONAL_BRACKET(LOCTEXT("SetTransformCurveKeys_Bracket", "Setting Transform Curve Keys"));
+			FBracket Bracket = ConditionalBracket(LOCTEXT("SetTransformCurveKeys_Bracket", "Setting Transform Curve Keys"), bShouldTransact);
 			
 			struct FKeys
 			{
@@ -798,11 +869,11 @@ bool UAnimDataController::SetTransformCurveKeys(const FAnimationCurveIdentifier&
 			
 			for (int32 SubCurveIndex = 0; SubCurveIndex < 3; ++SubCurveIndex)
 			{
-				const ETransformCurveChannel Channel = (ETransformCurveChannel)SubCurveIndex;
-				FKeys* CurveKeys = SubCurveKeys[SubCurveIndex];
+				const ETransformCurveChannel Channel = static_cast<ETransformCurveChannel>(SubCurveIndex);
+				const FKeys* CurveKeys = SubCurveKeys[SubCurveIndex];
 				for (int32 ChannelIndex = 0; ChannelIndex < 3; ++ChannelIndex)
 				{
-					const EVectorCurveChannel Axis = (EVectorCurveChannel)ChannelIndex;
+					const EVectorCurveChannel Axis = static_cast<EVectorCurveChannel>(ChannelIndex);
 					FAnimationCurveIdentifier TargetCurveIdentifier = CurveId;
 					UAnimationCurveIdentifierExtensions::GetTransformChildCurveIdentifier(TargetCurveIdentifier, Channel, Axis);
 					SetCurveKeys(TargetCurveIdentifier, CurveKeys->ChannelKeys[ChannelIndex], bShouldTransact);
@@ -830,11 +901,9 @@ bool UAnimDataController::SetTransformCurveKey(const FAnimationCurveIdentifier&
 {
 	ValidateModel();
 
-	FTransformCurve* Curve = Model->FindMutableTransformCurveById(CurveId);
-
-	if (Curve)
+	if (Model->FindMutableTransformCurveById(CurveId) != nullptr)
 	{
-		CONDITIONAL_BRACKET(LOCTEXT("AddTransformCurveKey_Bracket", "Setting Transform Curve Key"));
+		FBracket Bracket = ConditionalBracket(LOCTEXT("AddTransformCurveKey_Bracket", "Setting Transform Curve Key"), bShouldTransact);
 		struct FKeys
 		{
 			FRichCurveKey ChannelKeys[3];
@@ -861,11 +930,11 @@ bool UAnimDataController::SetTransformCurveKey(const FAnimationCurveIdentifier&
 		
 		for (int32 SubCurveIndex = 0; SubCurveIndex < 3; ++SubCurveIndex)
 		{
-			const ETransformCurveChannel Channel = (ETransformCurveChannel)SubCurveIndex;
+			const ETransformCurveChannel Channel = static_cast<ETransformCurveChannel>(SubCurveIndex);
 			const FKeys& VectorCurveKeys = VectorKeys[SubCurveIndex];
 			for (int32 ChannelIndex = 0; ChannelIndex < 3; ++ChannelIndex)
 			{
-				const EVectorCurveChannel Axis = (EVectorCurveChannel)ChannelIndex;
+				const EVectorCurveChannel Axis = static_cast<EVectorCurveChannel>(ChannelIndex);
 				FAnimationCurveIdentifier TargetCurveIdentifier = CurveId;
 				UAnimationCurveIdentifierExtensions::GetTransformChildCurveIdentifier(TargetCurveIdentifier, Channel, Axis);
 				SetCurveKey(TargetCurveIdentifier, VectorCurveKeys.ChannelKeys[ChannelIndex], bShouldTransact);
@@ -887,30 +956,27 @@ bool UAnimDataController::RemoveTransformCurveKey(const FAnimationCurveIdentifie
 {
 	ValidateModel();
 
-	FTransformCurve* TransformCurve = Model->FindMutableTransformCurveById(CurveId);
-	if (TransformCurve)
+	if (Model->FindMutableTransformCurveById(CurveId))
 	{
 		const FString BaseCurveName = CurveId.InternalName.DisplayName.ToString();
 		const TArray<FString> SubCurveNames = { TEXT( "Translation"), TEXT( "Rotation"), TEXT( "Scale") };
 		const TArray<FString> ChannelCurveNames = { TEXT("X"), TEXT("Y"), TEXT("Z") };
 
-		CONDITIONAL_BRACKET(LOCTEXT("RemoveTransformCurveKey_Bracket", "Deleting Animation Transform Curve Key"));
+		FBracket Bracket = ConditionalBracket(LOCTEXT("RemoveTransformCurveKey_Bracket", "Deleting Animation Transform Curve Key"), bShouldTransact);
 		
 		for (int32 SubCurveIndex = 0; SubCurveIndex < 3; ++SubCurveIndex)
 		{
-			const ETransformCurveChannel Channel = (ETransformCurveChannel)SubCurveIndex;
+			const ETransformCurveChannel Channel = static_cast<ETransformCurveChannel>(SubCurveIndex);
 			for (int32 ChannelIndex = 0; ChannelIndex < 3; ++ChannelIndex)
 			{
-				const EVectorCurveChannel Axis = (EVectorCurveChannel)ChannelIndex;
+				const EVectorCurveChannel Axis = static_cast<EVectorCurveChannel>(ChannelIndex);
 				FAnimationCurveIdentifier TargetCurveIdentifier = CurveId;
 				UAnimationCurveIdentifierExtensions::GetTransformChildCurveIdentifier(TargetCurveIdentifier, Channel, Axis);
 				RemoveCurveKey(TargetCurveIdentifier, Time, bShouldTransact);
 			}
 		}
 
-
 		return true;
-
 	}
 	else
 	{
@@ -930,10 +996,9 @@ bool UAnimDataController::RenameCurve(const FAnimationCurveIdentifier& CurveToRe
 		{
 			if (CurveToRenameId.CurveType == NewCurveId.CurveType)
 			{
-				FAnimCurveBase* Curve = Model->FindMutableCurveById(CurveToRenameId);
-				if (Curve)
+				if (FAnimCurveBase* Curve = Model->FindMutableCurveById(CurveToRenameId))
 				{
-					CONDITIONAL_TRANSACTION(LOCTEXT("RenameCurve", "Renaming Curve"));
+					FTransaction Transaction = ConditionalTransaction(LOCTEXT("RenameCurve", "Renaming Curve"), bShouldTransact);
 
 					FCurveRenamedPayload Payload;
 					Payload.Identifier = FAnimationCurveIdentifier(Curve->Name, CurveToRenameId.CurveType);
@@ -941,9 +1006,9 @@ bool UAnimDataController::RenameCurve(const FAnimationCurveIdentifier& CurveToRe
 					Curve->Name = NewCurveId.InternalName;
 					Payload.NewIdentifier = NewCurveId;
 
-					CONDITIONAL_ACTION(UE::Anim::FRenameCurveAction, NewCurveId, CurveToRenameId);
+					ConditionalAction<UE::Anim::FRenameCurveAction>(bShouldTransact, NewCurveId, CurveToRenameId);
 
-					Model->Notify(EAnimDataModelNotifyType::CurveRenamed, Payload);
+					Model->GetNotifier().Notify(EAnimDataModelNotifyType::CurveRenamed, Payload);
 
 					return true;
 				}
@@ -983,18 +1048,17 @@ bool UAnimDataController::SetCurveColor(const FAnimationCurveIdentifier& CurveId
 	{
 		if (CurveId.CurveType == ERawCurveTrackTypes::RCT_Float)
 		{
-			FFloatCurve* Curve = Model->FindMutableFloatCurveById(CurveId);
-			if (Curve)
+			if (FFloatCurve* Curve = Model->FindMutableFloatCurveById(CurveId))
 			{
-				CONDITIONAL_TRANSACTION(LOCTEXT("ChangingCurveColor", "Changing Curve Color"));
+				FTransaction Transaction = ConditionalTransaction(LOCTEXT("ChangingCurveColor", "Changing Curve Color"), bShouldTransact);
 
-				CONDITIONAL_ACTION(UE::Anim::FSetCurveColorAction, CurveId, Curve->Color);
+				ConditionalAction<UE::Anim::FSetCurveColorAction>(bShouldTransact, CurveId, Curve->Color);
 
 				Curve->Color = Color;
 
 				FCurveChangedPayload Payload;
 				Payload.Identifier = CurveId;
-				Model->Notify(EAnimDataModelNotifyType::CurveColorChanged, Payload);
+				Model->GetNotifier().Notify(EAnimDataModelNotifyType::CurveColorChanged, Payload);
 
 				return true;				
 			}
@@ -1021,13 +1085,12 @@ bool UAnimDataController::ScaleCurve(const FAnimationCurveIdentifier& CurveId, f
 {
 	ValidateModel();
 
-	ERawCurveTrackTypes SupportedCurveType = CurveId.CurveType;
+	const ERawCurveTrackTypes SupportedCurveType = CurveId.CurveType;
 	if (SupportedCurveType == ERawCurveTrackTypes::RCT_Float)
 	{
-		FFloatCurve* Curve = Model->FindMutableFloatCurveById(CurveId);
-		if (Curve)
+		if (FFloatCurve* Curve = Model->FindMutableFloatCurveById(CurveId))
 		{
-			CONDITIONAL_TRANSACTION(LOCTEXT("ScalingCurve", "Scaling Curve"));
+			FTransaction Transaction = ConditionalTransaction(LOCTEXT("ScalingCurve", "Scaling Curve"), bShouldTransact);
 
 			Curve->FloatCurve.ScaleCurve(Origin, Factor);
 
@@ -1036,9 +1099,9 @@ bool UAnimDataController::ScaleCurve(const FAnimationCurveIdentifier& CurveId, f
 			Payload.Factor = Factor;
 			Payload.Origin = Origin;
 			
-			CONDITIONAL_ACTION(UE::Anim::FScaleCurveAction, CurveId, Origin, 1.0f / Factor, SupportedCurveType);
+			ConditionalAction<UE::Anim::FScaleCurveAction>(bShouldTransact, CurveId, Origin, 1.0f / Factor, SupportedCurveType);
 
-			Model->Notify(EAnimDataModelNotifyType::CurveScaled, Payload);
+			Model->GetNotifier().Notify(EAnimDataModelNotifyType::CurveScaled, Payload);
 
 			return true;
 		}
@@ -1059,9 +1122,7 @@ bool UAnimDataController::SetCurveKey(const FAnimationCurveIdentifier& CurveId,
 {
 	ValidateModel();
 
-	FRichCurve* RichCurve = Model->GetMutableRichCurve(CurveId);
-	ERawCurveTrackTypes SupportedCurveType = CurveId.CurveType;
-	if (RichCurve)
+	if (FRichCurve* RichCurve = Model->GetMutableRichCurve(CurveId))
 	{
 		FCurveChangedPayload Payload;
 		Payload.Identifier = CurveId;
@@ -1070,25 +1131,25 @@ bool UAnimDataController::SetCurveKey(const FAnimationCurveIdentifier& CurveId,
 		const FKeyHandle Handle = RichCurve->FindKey(Key.Time, 0.f);
 		if (Handle != FKeyHandle::Invalid())
 		{
-			CONDITIONAL_TRANSACTION(LOCTEXT("SetNamedCurveKey", "Setting Curve Key"));
+			FTransaction Transaction = ConditionalTransaction(LOCTEXT("SetNamedCurveKey", "Setting Curve Key"), bShouldTransact);
 			// Cache old value for action
 			const FRichCurveKey CurrentKey = RichCurve->GetKey(Handle);
-			CONDITIONAL_ACTION(UE::Anim::FSetRichCurveKeyAction, CurveId, CurrentKey);
+			ConditionalAction<UE::Anim::FSetRichCurveKeyAction>(bShouldTransact, CurveId, CurrentKey);
 
 			// Set the new value
 			RichCurve->SetKeyValue(Handle, Key.Value);
 
-			Model->Notify(EAnimDataModelNotifyType::CurveChanged, Payload);
+			Model->GetNotifier().Notify(EAnimDataModelNotifyType::CurveChanged, Payload);
 		}
 		else
 		{
-			CONDITIONAL_TRANSACTION(LOCTEXT("AddNamedCurveKey", "Adding Curve Key"));
-			CONDITIONAL_ACTION(UE::Anim::FRemoveRichCurveKeyAction, CurveId, Key.Time);
+			FTransaction Transaction = ConditionalTransaction(LOCTEXT("AddNamedCurveKey", "Adding Curve Key"), bShouldTransact);
+			ConditionalAction<UE::Anim::FRemoveRichCurveKeyAction>(bShouldTransact, CurveId, Key.Time);
 
 			// Add the new key
 			RichCurve->AddKey(Key.Time, Key.Value);
 
-			Model->Notify(EAnimDataModelNotifyType::CurveChanged, Payload);
+			Model->GetNotifier().Notify(EAnimDataModelNotifyType::CurveChanged, Payload);
 		}
 
 		return true;
@@ -1101,9 +1162,7 @@ bool UAnimDataController::RemoveCurveKey(const FAnimationCurveIdentifier& CurveI
 {
 	ValidateModel();
 
-	FRichCurve* RichCurve = Model->GetMutableRichCurve(CurveId);
-	ERawCurveTrackTypes SupportedCurveType = CurveId.CurveType;
-	if (RichCurve)
+	if (FRichCurve* RichCurve = Model->GetMutableRichCurve(CurveId))
 	{
 		FCurveChangedPayload Payload;
 		Payload.Identifier = CurveId;
@@ -1112,15 +1171,15 @@ bool UAnimDataController::RemoveCurveKey(const FAnimationCurveIdentifier& CurveI
 		const FKeyHandle Handle = RichCurve->FindKey(Time, 0.f);
 		if (Handle != FKeyHandle::Invalid())
 		{
-			CONDITIONAL_TRANSACTION(LOCTEXT("RemoveNamedCurveKey", "Removing Curve Key"));
+			FTransaction Transaction = ConditionalTransaction(LOCTEXT("RemoveNamedCurveKey", "Removing Curve Key"), bShouldTransact);
 
 			// Cached current value for action
 			const FRichCurveKey CurrentKey = RichCurve->GetKey(Handle);
-			CONDITIONAL_ACTION(UE::Anim::FAddRichCurveKeyAction, CurveId, CurrentKey);
+			ConditionalAction<UE::Anim::FAddRichCurveKeyAction>(bShouldTransact, CurveId, CurrentKey);
 
 			RichCurve->DeleteKey(Handle);
 
-			Model->Notify(EAnimDataModelNotifyType::CurveChanged, Payload);
+			Model->GetNotifier().Notify(EAnimDataModelNotifyType::CurveChanged, Payload);
 
 			return true;
 		}
@@ -1138,19 +1197,17 @@ bool UAnimDataController::SetCurveKeys(const FAnimationCurveIdentifier& CurveId,
 {
 	ValidateModel();
 
-	FRichCurve* RichCurve = Model->GetMutableRichCurve(CurveId);
-	ERawCurveTrackTypes SupportedCurveType = CurveId.CurveType;
-	if (RichCurve)
+	if (FRichCurve* RichCurve = Model->GetMutableRichCurve(CurveId))
 	{
-		CONDITIONAL_TRANSACTION(LOCTEXT("SettingNamedCurveKeys", "Setting Curve Keys"));
-		CONDITIONAL_ACTION(UE::Anim::FSetRichCurveKeysAction, CurveId, RichCurve->GetConstRefOfKeys());
+		FTransaction Transaction = ConditionalTransaction(LOCTEXT("SettingNamedCurveKeys", "Setting Curve Keys"), bShouldTransact);
+		ConditionalAction<UE::Anim::FSetRichCurveKeysAction>(bShouldTransact, CurveId, RichCurve->GetConstRefOfKeys());
 
 		// Set rich curve values
 		RichCurve->SetKeys(CurveKeys);
 
 		FCurveChangedPayload Payload;
 		Payload.Identifier = CurveId;
-		Model->Notify(EAnimDataModelNotifyType::CurveChanged, Payload);
+		Model->GetNotifier().Notify(EAnimDataModelNotifyType::CurveChanged, Payload);
 
 		return true;
 	}
@@ -1161,152 +1218,42 @@ bool UAnimDataController::SetCurveKeys(const FAnimationCurveIdentifier& CurveId,
 void UAnimDataController::NotifyPopulated()
 {
 	ValidateModel();
-	Model->Notify(EAnimDataModelNotifyType::Populated);
+
+	Model->bPopulated = true;
+	Model->GetNotifier().Notify(EAnimDataModelNotifyType::Populated);
 }
 
 void UAnimDataController::NotifyBracketOpen()
 {
 	ValidateModel();
-	Model->Notify(EAnimDataModelNotifyType::BracketOpened);
+	Model->GetNotifier().Notify(EAnimDataModelNotifyType::BracketOpened);
 }
 
 void UAnimDataController::NotifyBracketClosed()
 {
 	ValidateModel();
-	Model->Notify(EAnimDataModelNotifyType::BracketClosed);
-}
-
-const bool UAnimDataController::IsSupportedCurveType(ERawCurveTrackTypes CurveType) const
-{
-	const TArray<ERawCurveTrackTypes> SupportedTypes = { ERawCurveTrackTypes::RCT_Float, ERawCurveTrackTypes::RCT_Transform };
-	return SupportedTypes.Contains(CurveType);
-}
-
-void UAnimDataController::ValidateModel() const
-{
-	checkf(Model != nullptr, TEXT("Invalid Model"));
+	Model->GetNotifier().Notify(EAnimDataModelNotifyType::BracketClosed);
 }
 
 void UAnimDataController::ResizePlayLength(float Length, float T0, float T1, bool bShouldTransact)
 {
-	const TRange<float> PlayRange(TRange<float>::BoundsType::Inclusive(0.f), TRange<float>::BoundsType::Inclusive(Model->PlayLength));
-	if (!FMath::IsNearlyZero(Length) && Length > 0.f)
-	{
-		if (Length != Model->PlayLength)
-		{
-			// Ensure that T0 is within the curent play range
-			if (PlayRange.Contains(T0))
-			{
-				// Ensure that the start and end length of either removal or insertion are valid
-				if (T0 < T1)
-				{
-					CONDITIONAL_TRANSACTION(LOCTEXT("ResizePlayLength", "Resizing Play Length"));
-
-					FSequenceLengthChangedPayload Payload;
-					Payload.T0 = T0;
-					Payload.T1 = T1;
-					Payload.PreviousLength = Model->PlayLength;
-
-					CONDITIONAL_ACTION(UE::Anim::FResizePlayLengthAction, Model.Get(), T0, T1);
-
-					Model->PlayLength = Length;
-
-					Model->NumberOfFrames = Model->FrameRate.AsFrameTime(Model->PlayLength).RoundToFrame().Value;
-					Model->NumberOfKeys = Model->NumberOfFrames + 1;
-	
-					Model->Notify<FSequenceLengthChangedPayload>(EAnimDataModelNotifyType::SequenceLengthChanged, Payload);
-				}
-				else
-				{
-					ReportErrorf(LOCTEXT("InvalidEndTimeError", "Invalid T1, smaller that T0 value: T0 {0}, T1 {1}"), FText::AsNumber(T0), FText::AsNumber(T1));
-				}
-			}
-			else
-			{
-				ReportErrorf(LOCTEXT("InvalidStartTimeError", "Invalid T0, not within existing play range: T0 {0}, Play Length {1}"), FText::AsNumber(T0), FText::AsNumber(Model->PlayLength));
-			}
-		}
-		else
-		{
-			ReportWarningf(LOCTEXT("SamePlayLengthWarning", "New play length is same as existing one: {0} seconds"), FText::AsNumber(Length));
-		}
-	}
-	else
-	{
-		ReportErrorf(LOCTEXT("InvalidPlayLengthError", "Invalid play length value provided: {0} seconds"), FText::AsNumber(Length));
-	}
-}
-
-
-void UAnimDataController::Report(ELogVerbosity::Type InVerbosity, const FText& InMessage) const
-{
-	FString Message = InMessage.ToString();
-	if (Model != nullptr)
-	{
-		if (UPackage* Package = Cast<UPackage>(Model->GetOutermost()))
-		{
-			Message = FString::Printf(TEXT("%s : %s"), *Package->GetPathName(), *Message);
-		}
-	}
-
-	FScriptExceptionHandler::Get().HandleException(InVerbosity, *Message, *FString());
-}
-
-FString UAnimDataController::GetCurveTypeValueName(ERawCurveTrackTypes InType) const
-{
-	FString ValueString;
-
-	const UEnum* Enum = FindObject<UEnum>(nullptr, TEXT("/Script/Engine.ERawCurveTrackTypes"));
-	if (Enum)
-	{
-		ValueString = Enum->GetNameStringByValue((int64)InType);
-	}
-
-	return ValueString;
-}
-
-bool UAnimDataController::CheckOuterClass(UClass* InClass) const
-{
-	ValidateModel();
-	
-	const UObject* ModelOuter = Model->GetOuter();
-	if (ModelOuter)
-	{
-		const UClass* OuterClass = ModelOuter->GetClass();
-		if (OuterClass)
-		{
-			if (OuterClass == InClass || OuterClass->IsChildOf(InClass))
-			{
-				return true;
-			}
-			else
-			{
-				ReportErrorf(LOCTEXT("NoValidOuterClassError", "Incorrect outer object class found for Animation Data Model {0}, expected {1} actual {2}"), FText::FromString(Model->GetName()), FText::FromString(InClass->GetName()), FText::FromString(OuterClass->GetName()));
-			}
-		}
-	}
-	else
-	{
-		ReportErrorf(LOCTEXT("NoValidOuterObjectFoundError", "No valid outer object found for Animation Data Model {0}"), FText::FromString(Model->GetName()));
-	}
-
-	return false;
+	ResizeNumberOfFrames(ConvertSecondsToFrameNumber(Length), ConvertSecondsToFrameNumber(T0), ConvertSecondsToFrameNumber(T1), bShouldTransact);
 }
 
 int32 UAnimDataController::AddBoneTrack(FName BoneName, bool bShouldTransact /*= true*/)
 {
-	if (!CheckOuterClass(UAnimSequence::StaticClass()))
+	if (!ModelInterface->GetAnimationSequence())
 	{
 		return INDEX_NONE;
 	}
 
-	CONDITIONAL_TRANSACTION(LOCTEXT("AddBoneTrack", "Adding Animation Data Track"));
+	FTransaction Transaction = ConditionalTransaction(LOCTEXT("AddBoneTrack", "Adding Animation Data Track"), bShouldTransact);
 	return InsertBoneTrack(BoneName, INDEX_NONE, bShouldTransact);
 }
 
 int32 UAnimDataController::InsertBoneTrack(FName BoneName, int32 DesiredIndex, bool bShouldTransact /*= true*/)
 {
-	if (!CheckOuterClass(UAnimSequence::StaticClass()))
+	if (!ModelInterface->GetAnimationSequence())
 	{
 		return INDEX_NONE;
 	}
@@ -1318,29 +1265,26 @@ int32 UAnimDataController::InsertBoneTrack(FName BoneName, int32 DesiredIndex, b
 		if (Model->GetNumBoneTracks() >= MAX_ANIMATION_TRACKS)
 		{
 			ReportWarningf(LOCTEXT("MaxNumberOfTracksReachedWarning", "Cannot add track with name {0}. An animation sequence cannot contain more than 65535 tracks"), FText::FromName(BoneName));
+			return INDEX_NONE;
 		}
 		else
 		{
-			CONDITIONAL_TRANSACTION(LOCTEXT("InsertBoneTrack", "Inserting Animation Data Track"));
+			FTransaction Transaction = ConditionalTransaction(LOCTEXT("InsertBoneTrack", "Inserting Animation Data Track"), bShouldTransact);
 
 			// Determine correct index to do insertion at
 			const int32 InsertIndex = Model->BoneAnimationTracks.IsValidIndex(DesiredIndex) ? DesiredIndex : Model->BoneAnimationTracks.Num();
-
-			FBoneAnimationTrack& NewTrack = Model->BoneAnimationTracks.InsertDefaulted_GetRef(InsertIndex);
-			NewTrack.Name = BoneName;
+			int32 BoneIndex = INDEX_NONE;
 
 			if (const UAnimSequence* AnimationSequence = Model->GetAnimationSequence())
 			{
 				if (const USkeleton* Skeleton = AnimationSequence->GetSkeleton())
 				{
-					const int32 BoneIndex = Skeleton->GetReferenceSkeleton().FindBoneIndex(BoneName);
+					BoneIndex = Skeleton->GetReferenceSkeleton().FindBoneIndex(BoneName);
 
 					if (BoneIndex == INDEX_NONE)
 					{
-						ReportWarningf(LOCTEXT("UnableToFindBoneIndexWarning", "Unable to retrieve bone index for track: {0}"), FText::FromName(BoneName));
+						ReportErrorf(LOCTEXT("UnableToFindBoneIndexWarning", "Unable to retrieve bone index for track: {0}"), FText::FromName(BoneName));
 					}
-
-					NewTrack.BoneTreeIndex = BoneIndex;
 				}
 				else
 				{
@@ -1352,16 +1296,23 @@ int32 UAnimDataController::InsertBoneTrack(FName BoneName, int32 DesiredIndex, b
 				Report(ELogVerbosity::Error, LOCTEXT("UnableToGetOuterAnimSequenceError", "Unable to retrieve outer Animation Sequence"));
 			}
 
+			if (BoneIndex != INDEX_NONE)
+			{
+				FBoneAnimationTrack& NewTrack = Model->BoneAnimationTracks.InsertDefaulted_GetRef(InsertIndex);
+				NewTrack.Name = BoneName;
+				NewTrack.BoneTreeIndex = BoneIndex;
+
 				FAnimationTrackAddedPayload Payload;
 				Payload.Name = BoneName;
 				Payload.TrackIndex = InsertIndex;
 
-			Model->Notify<FAnimationTrackAddedPayload>(EAnimDataModelNotifyType::TrackAdded, Payload);
-			CONDITIONAL_ACTION(UE::Anim::FRemoveTrackAction, NewTrack, InsertIndex);
+				Model->GetNotifier().Notify<FAnimationTrackAddedPayload>(EAnimDataModelNotifyType::TrackAdded, Payload);
+				ConditionalAction<UE::Anim::FRemoveTrackAction>(bShouldTransact, BoneName);
 
 				return InsertIndex;
 			}
 		}
+	}
 	else
 	{
 		ReportWarningf(LOCTEXT("TrackNameAlreadyExistsWarning", "Track with name {0} already exists"), FText::FromName(BoneName));
@@ -1372,7 +1323,7 @@ int32 UAnimDataController::InsertBoneTrack(FName BoneName, int32 DesiredIndex, b
 
 bool UAnimDataController::RemoveBoneTrack(FName BoneName, bool bShouldTransact /*= true*/)
 {
-	if (!CheckOuterClass(UAnimSequence::StaticClass()))
+	if (!ModelInterface->GetAnimationSequence())
 	{
 		return false;
 	}
@@ -1381,7 +1332,7 @@ bool UAnimDataController::RemoveBoneTrack(FName BoneName, bool bShouldTransact /
 
 	if (ExistingTrackPtr != nullptr)
 	{
-		CONDITIONAL_TRANSACTION(LOCTEXT("RemoveBoneTrack", "Removing Animation Data Track"));
+		FTransaction Transaction = ConditionalTransaction(LOCTEXT("RemoveBoneTrack", "Removing Animation Data Track"), bShouldTransact);
 		const int32 TrackIndex = Model->BoneAnimationTracks.IndexOfByPredicate([ExistingTrackPtr](const FBoneAnimationTrack& Track)
 		{
 			return Track.Name == ExistingTrackPtr->Name;
@@ -1389,13 +1340,13 @@ bool UAnimDataController::RemoveBoneTrack(FName BoneName, bool bShouldTransact /
 
 		ensure(TrackIndex != INDEX_NONE);
 
-		CONDITIONAL_ACTION(UE::Anim::FAddTrackAction, *ExistingTrackPtr, TrackIndex);
+		ConditionalAction<UE::Anim::FAddTrackAction>(bShouldTransact,*ExistingTrackPtr);
 		Model->BoneAnimationTracks.RemoveAt(TrackIndex);
 
 		FAnimationTrackRemovedPayload Payload;
 		Payload.Name = BoneName;
 
-		Model->Notify(EAnimDataModelNotifyType::TrackRemoved, Payload);
+		Model->GetNotifier().Notify(EAnimDataModelNotifyType::TrackRemoved, Payload);
 
 		return true;
 	}
@@ -1409,7 +1360,7 @@ bool UAnimDataController::RemoveBoneTrack(FName BoneName, bool bShouldTransact /
 
 void UAnimDataController::RemoveAllBoneTracks(bool bShouldTransact /*= true*/)
 {
-	if (!CheckOuterClass(UAnimSequence::StaticClass()))
+	if (!ModelInterface->GetAnimationSequence())
 	{
 		return;
 	}
@@ -1419,7 +1370,7 @@ void UAnimDataController::RemoveAllBoneTracks(bool bShouldTransact /*= true*/)
 
 	if (TrackNames.Num())
 	{
-		CONDITIONAL_BRACKET(LOCTEXT("RemoveAllBoneTracks", "Removing all Animation Data Tracks"));
+		FBracket Bracket = ConditionalBracket(LOCTEXT("RemoveAllBoneTracks", "Removing all Animation Data Tracks"), bShouldTransact);
 		for (const FName& TrackName : TrackNames)
 		{
 			RemoveBoneTrack(TrackName, bShouldTransact);
@@ -1429,13 +1380,11 @@ void UAnimDataController::RemoveAllBoneTracks(bool bShouldTransact /*= true*/)
 
 bool UAnimDataController::SetBoneTrackKeys(FName BoneName, const TArray<FVector>& PositionalKeys, const TArray<FQuat>& RotationalKeys, const TArray<FVector>& ScalingKeys, bool bShouldTransact /*= true*/)
 {
-	if (!CheckOuterClass(UAnimSequence::StaticClass()))
+	if (!ModelInterface->GetAnimationSequence())
 	{
 		return false;
 	}
 
-	CONDITIONAL_TRANSACTION(LOCTEXT("SetTrackKeysTransaction", "Setting Animation Data Track keys"));
-
 	// Validate key format
 	const int32 MaxNumKeys = FMath::Max(FMath::Max(PositionalKeys.Num(), RotationalKeys.Num()), ScalingKeys.Num());
 
@@ -1449,13 +1398,14 @@ bool UAnimDataController::SetBoneTrackKeys(FName BoneName, const TArray<FVector>
 		{
 			if (FBoneAnimationTrack* TrackPtr = Model->FindMutableBoneTrackByName(BoneName))
 			{
-				CONDITIONAL_ACTION(UE::Anim::FSetTrackKeysAction, *TrackPtr);
+				FTransaction Transaction = ConditionalTransaction(LOCTEXT("SetTrackKeysTransaction", "Setting Animation Data Track keys"), bShouldTransact);
 
+				ConditionalAction<UE::Anim::FSetTrackKeysAction>(bShouldTransact, *TrackPtr);
 
 				TrackPtr->InternalTrackData.PosKeys.SetNum(MaxNumKeys);
 				TrackPtr->InternalTrackData.ScaleKeys.SetNum(MaxNumKeys);
 				TrackPtr->InternalTrackData.RotKeys.SetNum(MaxNumKeys);
-				for(int KeyIndex = 0; KeyIndex<MaxNumKeys; KeyIndex++)
+				for(int32 KeyIndex = 0; KeyIndex<MaxNumKeys; KeyIndex++)
 				{
 					TrackPtr->InternalTrackData.PosKeys[KeyIndex] = FVector3f(PositionalKeys[KeyIndex]);
 					TrackPtr->InternalTrackData.ScaleKeys[KeyIndex] = FVector3f(ScalingKeys[KeyIndex]);
@@ -1465,7 +1415,7 @@ bool UAnimDataController::SetBoneTrackKeys(FName BoneName, const TArray<FVector>
 				FAnimationTrackChangedPayload Payload;
 				Payload.Name = BoneName;
 
-				Model->Notify(EAnimDataModelNotifyType::TrackChanged, Payload);
+				Model->GetNotifier().Notify(EAnimDataModelNotifyType::TrackChanged, Payload);
 
 				return true;
 			}
@@ -1494,7 +1444,7 @@ bool UAnimDataController::SetBoneTrackKeys(FName BoneName, const TArray<FVector3
 		return false;
 	}
 
-	CONDITIONAL_TRANSACTION(LOCTEXT("SetTrackKeysTransaction", "Setting Animation Data Track keys"));
+	FTransaction Transaction = ConditionalTransaction(LOCTEXT("SetTrackKeysTransaction", "Setting Animation Data Track keys"), bShouldTransact);
 
 	// Validate key format
 	const int32 MaxNumKeys = FMath::Max(FMath::Max(PositionalKeys.Num(), RotationalKeys.Num()), ScalingKeys.Num());
@@ -1509,7 +1459,7 @@ bool UAnimDataController::SetBoneTrackKeys(FName BoneName, const TArray<FVector3
 		{
 			if (FBoneAnimationTrack* TrackPtr = Model->FindMutableBoneTrackByName(BoneName))
 			{
-				CONDITIONAL_ACTION(UE::Anim::FSetTrackKeysAction, *TrackPtr);
+				ConditionalAction<UE::Anim::FSetTrackKeysAction>(bShouldTransact, *TrackPtr);
 
 				TrackPtr->InternalTrackData.PosKeys = PositionalKeys;
 				TrackPtr->InternalTrackData.RotKeys = RotationalKeys;
@@ -1518,7 +1468,7 @@ bool UAnimDataController::SetBoneTrackKeys(FName BoneName, const TArray<FVector3
 				FAnimationTrackChangedPayload Payload;
 				Payload.Name = BoneName;
 
-				Model->Notify(EAnimDataModelNotifyType::TrackChanged, Payload);
+				Model->GetNotifier().Notify(EAnimDataModelNotifyType::TrackChanged, Payload);
 
 				return true;
 			}
@@ -1540,23 +1490,23 @@ bool UAnimDataController::SetBoneTrackKeys(FName BoneName, const TArray<FVector3
 	return false;
 }
 
-int32 DiscreteInclusiveLower(const TRange<int32>& InRange)
+static int32 DiscreteInclusiveLower(const TRange<int32>& InRange)
 {
 	check(!InRange.GetLowerBound().IsOpen());
 
 	// Add one for exclusive lower bounds since they start on the next subsequent frame
-	static const int32 Offsets[]   = { 0, 1 };
+	static constexpr int32 Offsets[]   = { 0, 1 };
 	const int32        OffsetIndex = (int32)InRange.GetLowerBound().IsExclusive();
 
 	return InRange.GetLowerBound().GetValue() + Offsets[OffsetIndex];
 }
 
-int32 DiscreteExclusiveUpper(const TRange<int32>& InRange)
+static int32 DiscreteExclusiveUpper(const TRange<int32>& InRange)
 {
 	check(!InRange.GetUpperBound().IsOpen());
 
 	// Add one for inclusive upper bounds since they finish on the next subsequent frame
-	static const int32 Offsets[]   = { 0, 1 };
+	static constexpr int32 Offsets[]   = { 0, 1 };
 	const int32        OffsetIndex = (int32)InRange.GetUpperBound().IsInclusive();
 
 	return InRange.GetUpperBound().GetValue() + Offsets[OffsetIndex];
@@ -1568,9 +1518,6 @@ bool UAnimDataController::UpdateBoneTrackKeys(FName BoneName, const FInt32Range&
 	{
 		return false;
 	}
-
-	CONDITIONAL_TRANSACTION(LOCTEXT("SetTrackKeysRangeTransaction", "Setting Animation Data Track keys"));
-
 	// Validate key format
 	const int32 MaxNumKeys = FMath::Max(FMath::Max(PositionalKeys.Num(), RotationalKeys.Num()), ScalingKeys.Num());
 	const int32 RangeMin = DiscreteInclusiveLower(KeyRangeToSet);
@@ -1596,7 +1543,9 @@ bool UAnimDataController::UpdateBoneTrackKeys(FName BoneName, const FInt32Range&
 						TArray<FQuat4f>& TrackRotKeys = InternalTrackData.RotKeys;
 						TArray<FVector3f>& TrackScaleKeys = InternalTrackData.ScaleKeys;
 
-						CONDITIONAL_ACTION(UE::Anim::FSetTrackKeysAction, *TrackPtr);
+						FTransaction Transaction = ConditionalTransaction(LOCTEXT("SetTrackKeysRangeTransaction", "Setting Animation Data Track keys"), bShouldTransact);
+
+						ConditionalAction<UE::Anim::FSetTrackKeysAction>(bShouldTransact, *TrackPtr);
 
 						int32 KeyIndex = 0;
 						for (int32 FrameIndex = RangeMin; FrameIndex < RangeMax; ++FrameIndex, ++KeyIndex)
@@ -1609,7 +1558,7 @@ bool UAnimDataController::UpdateBoneTrackKeys(FName BoneName, const FInt32Range&
 						FAnimationTrackChangedPayload Payload;
 						Payload.Name = BoneName;
 
-						Model->Notify(EAnimDataModelNotifyType::TrackChanged, Payload);
+						Model->GetNotifier().Notify(EAnimDataModelNotifyType::TrackChanged, Payload);
 
 						return true;
 					}
@@ -1652,8 +1601,6 @@ bool UAnimDataController::UpdateBoneTrackKeys(FName BoneName, const FInt32Range&
 		return false;
 	}
 
-	CONDITIONAL_TRANSACTION(LOCTEXT("SetTrackKeysTransaction", "Setting Animation Data Track keys"));
-
 	// Validate key format
 	const int32 MaxNumKeys = FMath::Max(FMath::Max(PositionalKeys.Num(), RotationalKeys.Num()), ScalingKeys.Num());
 	const int32 RangeMin = DiscreteInclusiveLower(KeyRangeToSet);
@@ -1679,7 +1626,9 @@ bool UAnimDataController::UpdateBoneTrackKeys(FName BoneName, const FInt32Range&
 						TArray<FQuat4f>& TrackRotKeys = InternalTrackData.RotKeys;
 						TArray<FVector3f>& TrackScaleKeys = InternalTrackData.ScaleKeys;
 
-						CONDITIONAL_ACTION(UE::Anim::FSetTrackKeysAction, *TrackPtr);
+						FTransaction Transaction = ConditionalTransaction(LOCTEXT("SetTrackKeysTransaction", "Setting Animation Data Track keys"), bShouldTransact);
+
+						ConditionalAction<UE::Anim::FSetTrackKeysAction>(bShouldTransact, *TrackPtr);
 
 						int32 KeyIndex = 0;
 						for (int32 FrameIndex = RangeMin; FrameIndex < RangeMax; ++FrameIndex, ++KeyIndex)
@@ -1692,7 +1641,7 @@ bool UAnimDataController::UpdateBoneTrackKeys(FName BoneName, const FInt32Range&
 						FAnimationTrackChangedPayload Payload;
 						Payload.Name = BoneName;
 
-						Model->Notify(EAnimDataModelNotifyType::TrackChanged, Payload);
+						Model->GetNotifier().Notify(EAnimDataModelNotifyType::TrackChanged, Payload);
 
 						return true;
 					}
@@ -1730,7 +1679,7 @@ bool UAnimDataController::UpdateBoneTrackKeys(FName BoneName, const FInt32Range&
 
 void UAnimDataController::ResizeCurves(float NewLength, bool bInserted, float T0, float T1, bool bShouldTransact /*= true*/)
 {
-	CONDITIONAL_BRACKET(LOCTEXT("ResizeCurves", "Resizing all Curves"));
+	FBracket Bracket = ConditionalBracket(LOCTEXT("ResizeCurves", "Resizing all Curves"), bShouldTransact);
 
 	for (FFloatCurve& Curve : Model->CurveData.FloatCurves)
 	{
@@ -1744,11 +1693,11 @@ void UAnimDataController::ResizeCurves(float NewLength, bool bInserted, float T0
 		FTransformCurve ResizedCurve = Curve;
 		for (int32 SubCurveIndex = 0; SubCurveIndex < 3; ++SubCurveIndex)
 		{
-			const ETransformCurveChannel Channel = (ETransformCurveChannel)SubCurveIndex;
+			const ETransformCurveChannel Channel = static_cast<ETransformCurveChannel>(SubCurveIndex);
 			FVectorCurve& SubCurve = *ResizedCurve.GetVectorCurveByIndex(SubCurveIndex);
 			for (int32 ChannelIndex = 0; ChannelIndex < 3; ++ChannelIndex)
 			{
-				const EVectorCurveChannel Axis = (EVectorCurveChannel)ChannelIndex;
+				const EVectorCurveChannel Axis = static_cast<EVectorCurveChannel>(ChannelIndex);
 				FAnimationCurveIdentifier TargetCurveIdentifier = FAnimationCurveIdentifier(Curve.Name, ERawCurveTrackTypes::RCT_Transform);
 				UAnimationCurveIdentifierExtensions::GetTransformChildCurveIdentifier(TargetCurveIdentifier, Channel, Axis);
 				
@@ -1762,7 +1711,7 @@ void UAnimDataController::ResizeCurves(float NewLength, bool bInserted, float T0
 
 void UAnimDataController::ResizeAttributes(float NewLength, bool bInserted, float T0, float T1, bool bShouldTransact)
 {
-	CONDITIONAL_BRACKET(LOCTEXT("ResizeAttributes", "Resizing all Attributes"));
+	FBracket Bracket = ConditionalBracket(LOCTEXT("ResizeAttributes", "Resizing all Attributes"), bShouldTransact);
 
 	for (FAnimatedBoneAttribute& Attribute : Model->AnimatedBoneAttributes)
 	{
@@ -1794,18 +1743,18 @@ bool UAnimDataController::AddAttribute(const FAnimationAttributeIdentifier& Attr
 
 		if (!bAttributeAlreadyExists)
 		{
-			CONDITIONAL_TRANSACTION(LOCTEXT("AddAttribute", "Adding Animated Bone Attribute"));
+			FTransaction Transaction = ConditionalTransaction(LOCTEXT("AddAttribute", "Adding Animated Bone Attribute"), bShouldTransact);
 
 			FAnimatedBoneAttribute& Attribute = Model->AnimatedBoneAttributes.AddDefaulted_GetRef();
 			Attribute.Identifier = AttributeIdentifier;
 
 			Attribute.Curve.SetScriptStruct(AttributeIdentifier.GetType());
 		
-			CONDITIONAL_ACTION(UE::Anim::FRemoveAtributeAction, AttributeIdentifier);
+			ConditionalAction<UE::Anim::FRemoveAtributeAction>(bShouldTransact, AttributeIdentifier);
 
 			FAttributeAddedPayload Payload;
 			Payload.Identifier = AttributeIdentifier;
-			Model->Notify(EAnimDataModelNotifyType::AttributeAdded, Payload);
+			Model->GetNotifier().Notify(EAnimDataModelNotifyType::AttributeAdded, Payload);
 
 			return true;
 		}
@@ -1834,15 +1783,15 @@ bool UAnimDataController::RemoveAttribute(const FAnimationAttributeIdentifier& A
 
 		if (AttributeIndex != INDEX_NONE)
 		{
-			CONDITIONAL_TRANSACTION(LOCTEXT("RemoveAttribute", "Removing Animated Bone Attribute"));
+			FTransaction Transaction = ConditionalTransaction(LOCTEXT("RemoveAttribute", "Removing Animated Bone Attribute"), bShouldTransact);
 
-			CONDITIONAL_ACTION(UE::Anim::FAddAtributeAction, Model->AnimatedBoneAttributes[AttributeIndex]);
+			ConditionalAction<UE::Anim::FAddAtributeAction>(bShouldTransact, Model->AnimatedBoneAttributes[AttributeIndex]);
 
 			Model->AnimatedBoneAttributes.RemoveAtSwap(AttributeIndex);
 			
 			FAttributeRemovedPayload Payload;
 			Payload.Identifier = AttributeIdentifier;
-			Model->Notify(EAnimDataModelNotifyType::AttributeRemoved, Payload);
+			Model->GetNotifier().Notify(EAnimDataModelNotifyType::AttributeRemoved, Payload);
 
 			return true;
 		}
@@ -1879,7 +1828,7 @@ int32 UAnimDataController::RemoveAllAttributesForBone(const FName& BoneName, boo
 
 	if (Identifiers.Num())
 	{
-		CONDITIONAL_BRACKET(LOCTEXT("RemoveAllAttributesForBone", "Removing all Attributes for Bone"));
+		FBracket Bracket = ConditionalBracket(LOCTEXT("RemoveAllAttributesForBone", "Removing all Attributes for Bone"), bShouldTransact);
 		for (const FAnimationAttributeIdentifier& Identifier : Identifiers)
 		{
 			NumRemovedAttributes += RemoveAttribute(Identifier, bShouldTransact) ? 1 : 0;
@@ -1902,7 +1851,7 @@ int32 UAnimDataController::RemoveAllAttributes(bool bShouldTransact)
 
 	if (Identifiers.Num())
 	{
-		CONDITIONAL_BRACKET(LOCTEXT("RemoveAllAttributes", "Removing all Attributes"));
+		FBracket Bracket = ConditionalBracket(LOCTEXT("RemoveAllAttributes", "Removing all Attributes"), bShouldTransact);
 		for (const FAnimationAttributeIdentifier& Identifier : Identifiers)
 		{
 			NumRemovedAttributes += RemoveAttribute(Identifier, bShouldTransact) ? 1: 0;
@@ -1918,7 +1867,7 @@ bool UAnimDataController::SetAttributeKey_Internal(const FAnimationAttributeIden
 	{
 		if (KeyValue)
 		{
-			FAnimatedBoneAttribute* AttributePtr = Model->AnimatedBoneAttributes.FindByPredicate([AttributeIdentifier](FAnimatedBoneAttribute& Attribute)
+			FAnimatedBoneAttribute* AttributePtr = Model->AnimatedBoneAttributes.FindByPredicate([AttributeIdentifier](const FAnimatedBoneAttribute& Attribute)
 			{
 				return Attribute.Identifier == AttributeIdentifier;
 			});
@@ -1927,26 +1876,26 @@ bool UAnimDataController::SetAttributeKey_Internal(const FAnimationAttributeIden
 			{
 				if (TypeStruct == AttributePtr->Identifier.GetType())
 				{
-					CONDITIONAL_TRANSACTION(LOCTEXT("SettingAttributeKey", "Setting Animated Bone Attribute key"));
+					FTransaction Transaction = ConditionalTransaction(LOCTEXT("SettingAttributeKey", "Setting Animated Bone Attribute key"), bShouldTransact);
 
 					FAttributeCurve& Curve = AttributePtr->Curve;
-					FKeyHandle KeyHandle = Curve.FindKey(Time);
+					const FKeyHandle KeyHandle = Curve.FindKey(Time);
 					// In case the key does not yet exist one will be added, and thus the undo is a remove
 					if (KeyHandle == FKeyHandle::Invalid())
 					{
-						CONDITIONAL_ACTION(UE::Anim::FRemoveAtributeKeyAction, AttributeIdentifier, Time);
-						Curve.UpdateOrAddKey(Time, KeyValue);
+						ConditionalAction<UE::Anim::FRemoveAtributeKeyAction>(bShouldTransact, AttributeIdentifier, Time);
+						Curve.UpdateOrAddTypedKey(Time, KeyValue, TypeStruct);
 					}
 					// In case the key does exist it will be updated , and thus the undo is a revert to the current value
 					else
 					{
-						CONDITIONAL_ACTION(UE::Anim::FSetAtributeKeyAction, AttributeIdentifier, Curve.GetKey(KeyHandle));
-						Curve.UpdateOrAddKey(Time, KeyValue);
+						ConditionalAction<UE::Anim::FSetAtributeKeyAction>(bShouldTransact, AttributeIdentifier, Curve.GetKey(KeyHandle));
+						Curve.UpdateOrAddTypedKey(Time, KeyValue, TypeStruct);
 					}
 
 					FAttributeChangedPayload Payload;
 					Payload.Identifier = AttributeIdentifier;
-					Model->Notify(EAnimDataModelNotifyType::AttributeChanged, Payload);
+					Model->GetNotifier().Notify(EAnimDataModelNotifyType::AttributeChanged, Payload);
 
 					return true;
 				}
@@ -1981,7 +1930,7 @@ bool UAnimDataController::SetAttributeKeys_Internal(const FAnimationAttributeIde
 	{
 		if (Times.Num() == KeyValues.Num())
 		{
-			FAnimatedBoneAttribute* AttributePtr = Model->AnimatedBoneAttributes.FindByPredicate([AttributeIdentifier](FAnimatedBoneAttribute& Attribute)
+			FAnimatedBoneAttribute* AttributePtr = Model->AnimatedBoneAttributes.FindByPredicate([AttributeIdentifier](const FAnimatedBoneAttribute& Attribute)
 			{
 				return Attribute.Identifier == AttributeIdentifier;
 			});
@@ -1990,17 +1939,17 @@ bool UAnimDataController::SetAttributeKeys_Internal(const FAnimationAttributeIde
 			{
 				if (TypeStruct == AttributePtr->Identifier.GetType())
 				{
-					CONDITIONAL_TRANSACTION(LOCTEXT("SettingAttributeKeys", "Setting Animated Bone Attribute keys"));
+					FTransaction Transaction = ConditionalTransaction(LOCTEXT("SettingAttributeKeys", "Setting Animated Bone Attribute keys"), bShouldTransact);
 
 					FAnimatedBoneAttribute& Attribute = *AttributePtr;
 
-					CONDITIONAL_ACTION(UE::Anim::FSetAtributeKeysAction, Attribute);
+					ConditionalAction<UE::Anim::FSetAtributeKeysAction>(bShouldTransact, Attribute);
 			
 					Attribute.Curve.SetKeys(Times, KeyValues);
 
 					FAttributeChangedPayload Payload;
 					Payload.Identifier = AttributeIdentifier;
-					Model->Notify(EAnimDataModelNotifyType::AttributeChanged, Payload);
+					Model->GetNotifier().Notify(EAnimDataModelNotifyType::AttributeChanged, Payload);
 
 					return true;
 				}
@@ -2043,19 +1992,19 @@ bool UAnimDataController::RemoveAttributeKey(const FAnimationAttributeIdentifier
 		if (AttributePtr)
 		{
 			FAttributeCurve& Curve = AttributePtr->Curve;
-			FKeyHandle KeyHandle = Curve.FindKey(Time);
+			const FKeyHandle KeyHandle = Curve.FindKey(Time);
 
 			if (KeyHandle != FKeyHandle::Invalid())
 			{
-				CONDITIONAL_TRANSACTION(LOCTEXT("RemovingAttributeKey", "Removing Animated Bone Attribute key"));
+				FTransaction Transaction = ConditionalTransaction(LOCTEXT("RemovingAttributeKey", "Removing Animated Bone Attribute key"), bShouldTransact);
 
-				CONDITIONAL_ACTION(UE::Anim::FAddAtributeKeyAction, AttributeIdentifier, Curve.GetKey(KeyHandle));
+				ConditionalAction<UE::Anim::FAddAtributeKeyAction>(bShouldTransact, AttributeIdentifier, Curve.GetKey(KeyHandle));
 
 				Curve.DeleteKey(KeyHandle);
 
 				FAttributeAddedPayload Payload;
 				Payload.Identifier = AttributeIdentifier;
-				Model->Notify(EAnimDataModelNotifyType::AttributeChanged, Payload);
+				Model->GetNotifier().Notify(EAnimDataModelNotifyType::AttributeChanged, Payload);
 
 				return true;
 			}
@@ -2091,7 +2040,7 @@ bool UAnimDataController::DuplicateAttribute(const FAnimationAttributeIdentifier
 			{
 				if(const FAnimatedBoneAttribute* AttributePtr = Model->FindAttribute(AttributeIdentifier))
 				{
-					CONDITIONAL_TRANSACTION(LOCTEXT("DuplicateAttribute", "Duplicating Animation Attribute"));
+					FTransaction Transaction = ConditionalTransaction(LOCTEXT("DuplicateAttribute", "Duplicating Animation Attribute"), bShouldTransact);
 
 					FAnimatedBoneAttribute& DuplicateAttribute = Model->AnimatedBoneAttributes.AddDefaulted_GetRef();
 					DuplicateAttribute.Identifier = NewAttributeIdentifier;
@@ -2099,9 +2048,9 @@ bool UAnimDataController::DuplicateAttribute(const FAnimationAttributeIdentifier
 
 					FAttributeAddedPayload Payload;
 					Payload.Identifier = NewAttributeIdentifier;
-					Model->Notify(EAnimDataModelNotifyType::AttributeAdded, Payload);
+					Model->GetNotifier().Notify(EAnimDataModelNotifyType::AttributeAdded, Payload);
 
-					CONDITIONAL_ACTION(UE::Anim::FRemoveAtributeAction, NewAttributeIdentifier);
+					ConditionalAction<UE::Anim::FRemoveAtributeAction>(bShouldTransact, NewAttributeIdentifier);
 					
 					return true;
 				}
@@ -2128,6 +2077,28 @@ bool UAnimDataController::DuplicateAttribute(const FAnimationAttributeIdentifier
 	return false;
 }
 
+void UAnimDataController::UpdateWithSkeleton(USkeleton* TargetSkeleton, bool bShouldTransact)
+{
+	RemoveBoneTracksMissingFromSkeleton(TargetSkeleton);
+}
+
+void UAnimDataController::PopulateWithExistingModel(TScriptInterface<IAnimationDataModel> InModel)
+{
+	Model->BoneAnimationTracks = InModel->GetBoneAnimationTracks();
+	Model->FrameRate = InModel->GetFrameRate();
+	Model->NumberOfFrames = InModel->GetNumberOfFrames();
+	Model->NumberOfKeys = InModel->GetNumberOfFrames() + 1;
+	Model->CurveData = InModel->GetCurveData();
+	Model->AnimatedBoneAttributes = InModel->GetAttributes();
+
+	// Ensure curve name UIDs are validated against the outer skeleton as they are not copied over through FSmartName
+	if (const UAnimSequenceBase* AnimSequenceBase = Cast<UAnimSequenceBase>(InModel.GetObject()->GetOuter()))
+	{
+		FindOrAddCurveNamesOnSkeleton(AnimSequenceBase->GetSkeleton(), ERawCurveTrackTypes::RCT_Float);
+		FindOrAddCurveNamesOnSkeleton(AnimSequenceBase->GetSkeleton(), ERawCurveTrackTypes::RCT_Transform);
+	}
+}
+
 #endif // WITH_EDITOR
 
 #undef LOCTEXT_NAMESPACE // "AnimDataController"
