ANIM-RT truncation fixes - AnimationDataController

diff --git a/Engine/Source/Developer/AnimationDataController/Private/AnimDataController.cpp b/Engine/Source/Developer/AnimationDataController/Private/AnimDataController.cpp
--- a/Engine/Source/Developer/AnimationDataController/Private/AnimDataController.cpp
+++ b/Engine/Source/Developer/AnimationDataController/Private/AnimDataController.cpp
@@ -117,9 +117,9 @@ void UAnimDataController::ResizeNumberOfFrames(FFrameNumber NewLength, FFrameNum
 
 					FSequenceLengthChangedPayload Payload;
 					PRAGMA_DISABLE_DEPRECATION_WARNINGS
-					Payload.T0 = CurrentFrameRate.AsSeconds(T0);
-					Payload.T1 = CurrentFrameRate.AsSeconds(T1);
-					Payload.PreviousLength = CurrentFrameRate.AsSeconds(CurrentNumberOfFrames);
+					Payload.T0 = static_cast<float>(CurrentFrameRate.AsSeconds(T0));
+					Payload.T1 = static_cast<float>(CurrentFrameRate.AsSeconds(T1));
+					Payload.PreviousLength = static_cast<float>(CurrentFrameRate.AsSeconds(CurrentNumberOfFrames));
 					PRAGMA_ENABLE_DEPRECATION_WARNINGS
 
 					Payload.PreviousNumberOfFrames = CurrentNumberOfFrames;
@@ -177,8 +177,8 @@ void UAnimDataController::ResizeInFrames(FFrameNumber NewLength, FFrameNumber T0
 					ResizeNumberOfFrames(NewLength, T0, T1, bShouldTransact);
 					
 					const FFrameRate CurrentFrameRate = Model->GetFrameRate();
-					ResizeCurves(CurrentFrameRate.AsSeconds(NewLength), bInserted, CurrentFrameRate.AsSeconds(T0), CurrentFrameRate.AsSeconds(T1), bShouldTransact);
-					ResizeAttributes(CurrentFrameRate.AsSeconds(NewLength), bInserted, CurrentFrameRate.AsSeconds(T0), CurrentFrameRate.AsSeconds(T1), bShouldTransact);
+					ResizeCurves(static_cast<float>(CurrentFrameRate.AsSeconds(NewLength)), bInserted, static_cast<float>(CurrentFrameRate.AsSeconds(T0)), static_cast<float>(CurrentFrameRate.AsSeconds(T1)), bShouldTransact);
+					ResizeAttributes(static_cast<float>(CurrentFrameRate.AsSeconds(NewLength)), bInserted, static_cast<float>(CurrentFrameRate.AsSeconds(T0)), static_cast<float>(CurrentFrameRate.AsSeconds(T1)), bShouldTransact);
 				}
 				else
                 {
@@ -860,7 +860,7 @@ bool UAnimDataController::SetTransformCurveKeys(const FAnimationCurveIdentifier&
 				{
 					for (int32 ChannelIndex = 0; ChannelIndex < 3; ++ChannelIndex)
 					{
-						Key.ChannelKeys[ChannelIndex][KeyIndex] = FRichCurveKey(Time, Vector[ChannelIndex]);
+						Key.ChannelKeys[ChannelIndex][KeyIndex] = FRichCurveKey(Time, static_cast<float>(Vector[ChannelIndex]));
 					}
 				};
 
@@ -922,7 +922,7 @@ bool UAnimDataController::SetTransformCurveKey(const FAnimationCurveIdentifier&
 		{
 			for (int32 ChannelIndex = 0; ChannelIndex < 3; ++ChannelIndex)
 			{
-				Key.ChannelKeys[ChannelIndex] = FRichCurveKey(Time, Vector[ChannelIndex]);
+				Key.ChannelKeys[ChannelIndex] = FRichCurveKey(Time, static_cast<float>(Vector[ChannelIndex]));
 			}
 		};
 
