Animation Curve Runtime & Editor Improvements

diff --git a/Engine/Plugins/Interchange/Runtime/Source/FactoryNodes/Private/InterchangeAnimSequenceFactoryNode.cpp b/Engine/Plugins/Interchange/Runtime/Source/FactoryNodes/Private/InterchangeAnimSequenceFactoryNode.cpp
--- a/Engine/Plugins/Interchange/Runtime/Source/FactoryNodes/Private/InterchangeAnimSequenceFactoryNode.cpp
+++ b/Engine/Plugins/Interchange/Runtime/Source/FactoryNodes/Private/InterchangeAnimSequenceFactoryNode.cpp
@@ -170,6 +170,16 @@ bool UInterchangeAnimSequenceFactoryNode::GetCustomDoNotImportCurveWithZero(bool
 	IMPLEMENT_NODE_ATTRIBUTE_GETTER(DoNotImportCurveWithZero, bool);
 }
 
+bool UInterchangeAnimSequenceFactoryNode::SetCustomAddCurveMetadataToSkeleton(const bool& AttributeValue)
+{
+	IMPLEMENT_NODE_ATTRIBUTE_SETTER_NODELEGATE(AddCurveMetadataToSkeleton, bool);
+}
+
+bool UInterchangeAnimSequenceFactoryNode::GetCustomAddCurveMetadataToSkeleton(bool& AttributeValue) const
+{
+	IMPLEMENT_NODE_ATTRIBUTE_GETTER(AddCurveMetadataToSkeleton, bool);
+}
+
 bool UInterchangeAnimSequenceFactoryNode::SetCustomDoNotImportCurveWithZero(const bool& AttributeValue)
 {
 	IMPLEMENT_NODE_ATTRIBUTE_SETTER_NODELEGATE(DoNotImportCurveWithZero, bool);
diff --git a/Engine/Plugins/Interchange/Runtime/Source/FactoryNodes/Public/InterchangeAnimSequenceFactoryNode.h b/Engine/Plugins/Interchange/Runtime/Source/FactoryNodes/Public/InterchangeAnimSequenceFactoryNode.h
--- a/Engine/Plugins/Interchange/Runtime/Source/FactoryNodes/Public/InterchangeAnimSequenceFactoryNode.h
+++ b/Engine/Plugins/Interchange/Runtime/Source/FactoryNodes/Public/InterchangeAnimSequenceFactoryNode.h
@@ -157,6 +157,22 @@ public:
 	UFUNCTION(BlueprintCallable, Category = "Interchange | Node | AnimSequence")
 	bool SetCustomDoNotImportCurveWithZero(const bool& AttributeValue);
 
+	/**
+	 * Get the custom attribute AddCurveMetadataToSkeleton, return false if the attribute is not set.
+	 * 
+	 * Note - If value is true, do not import if it doesn't have any value other than zero. This is to avoid dirtying the skeleton on import.
+	 */
+	UFUNCTION(BlueprintCallable, Category = "Interchange | Node | AnimSequence")
+	bool GetCustomAddCurveMetadataToSkeleton(bool& AttributeValue) const;
+
+	/**
+	 * Set the custom attribute AddCurveMetadataToSkeleton. Return false if the attribute cannot be set.
+	 * 
+	 * Note - If value is true, do not import if it doesn't have any value other than zero. This is to avoid dirtying the skeleton on import.
+	 */
+	UFUNCTION(BlueprintCallable, Category = "Interchange | Node | AnimSequence")
+	bool SetCustomAddCurveMetadataToSkeleton(const bool& AttributeValue);
+	
 	/**
 	 * Get the custom attribute RemoveCurveRedundantKeys, return false if the attribute is not set.
 	 */
@@ -423,6 +439,7 @@ private:
 
 	const UE::Interchange::FAttributeKey Macro_CustomImportAttributeCurvesKey = UE::Interchange::FAttributeKey(TEXT("ImportAttributeCurves"));
 	const UE::Interchange::FAttributeKey Macro_CustomDoNotImportCurveWithZeroKey = UE::Interchange::FAttributeKey(TEXT("DoNotImportCurveWithZero"));
+	const UE::Interchange::FAttributeKey Macro_CustomAddCurveMetadataToSkeletonKey = UE::Interchange::FAttributeKey(TEXT("AddCurveMetadataToSkeleton"));
 	const UE::Interchange::FAttributeKey Macro_CustomRemoveCurveRedundantKeysKey = UE::Interchange::FAttributeKey(TEXT("RemoveCurveRedundantKeys"));
 	const UE::Interchange::FAttributeKey Macro_CustomMaterialDriveParameterOnCustomAttributeKey = UE::Interchange::FAttributeKey(TEXT("MaterialDriveParameterOnCustomAttribute"));
 	const UE::Interchange::FAttributeKey Macro_CustomDeleteExistingMorphTargetCurvesKey = UE::Interchange::FAttributeKey(TEXT("DeleteExistingMorphTargetCurves"));
diff --git a/Engine/Plugins/Interchange/Runtime/Source/Import/Private/Animation/InterchangeAnimSequenceFactory.cpp b/Engine/Plugins/Interchange/Runtime/Source/Import/Private/Animation/InterchangeAnimSequenceFactory.cpp
--- a/Engine/Plugins/Interchange/Runtime/Source/Import/Private/Animation/InterchangeAnimSequenceFactory.cpp
+++ b/Engine/Plugins/Interchange/Runtime/Source/Import/Private/Animation/InterchangeAnimSequenceFactory.cpp
@@ -152,6 +152,7 @@ namespace UE::Interchange::Private
 		, const FString& CurveName
 		, const int32 CurveFlags
 		, const bool bDoNotImportCurveWithZero
+		, const bool bAddCurveMetadataToSkeleton
 		, const bool bMorphTargetCurve
 		, const bool bMaterialCurve
 		, const bool bShouldTransact)
@@ -183,22 +184,10 @@ namespace UE::Interchange::Private
 				return false;
 			}
 		}
+		
 		FName Name = *CurveName;
-		USkeleton* Skeleton = TargetSequence->GetSkeleton();
-		const FSmartNameMapping* NameMapping = Skeleton->GetSmartNameContainer(USkeleton::AnimCurveMappingName);
-
-		// Add or retrieve curve
-		if (!NameMapping->Exists(Name))
-		{
-			// mark skeleton dirty
-			Skeleton->Modify();
-		}
-
-		FSmartName NewName;
-		Skeleton->AddSmartNameAndModify(USkeleton::AnimCurveMappingName, Name, NewName);
-
-		FAnimationCurveIdentifier FloatCurveId(NewName, ERawCurveTrackTypes::RCT_Float);
-
+		FAnimationCurveIdentifier FloatCurveId(Name, ERawCurveTrackTypes::RCT_Float);
+		
 		IAnimationDataModel* DataModel = TargetSequence->GetDataModel();
 		IAnimationDataController& Controller = TargetSequence->GetController();
 
@@ -218,8 +207,6 @@ namespace UE::Interchange::Private
 		// Should be valid at this point
 		ensure(TargetCurve);
 
-		Controller.UpdateCurveNamesFromSkeleton(Skeleton, ERawCurveTrackTypes::RCT_Float, bShouldTransact);
-
 		//MorphTarget curves are shift to 0 second
 		const float CurveStartTime = CurvePayload.Curves[0].GetFirstKey().Time;
 		if (CurveStartTime > UE_SMALL_NUMBER)
@@ -231,31 +218,35 @@ namespace UE::Interchange::Private
 
 		if (bMaterialCurve || bMorphTargetCurve)
 		{
-			Skeleton->AccumulateCurveMetaData(*CurveName, bMaterialCurve, bMorphTargetCurve);
+			if(bAddCurveMetadataToSkeleton)
+			{
+				USkeleton* Skeleton = TargetSequence->GetSkeleton();
+				Skeleton->AccumulateCurveMetaData(Name, bMaterialCurve, bMorphTargetCurve);
+			}
 		}
 		return true;
 	}
 
-	bool CreateMorphTargetCurve(UAnimSequence* TargetSequence, FAnimationCurvePayloadData& CurvePayload, const FString& CurveName, int32 CurveFlags, bool bDoNotImportCurveWithZero, bool bShouldTransact)
+	bool CreateMorphTargetCurve(UAnimSequence* TargetSequence, FAnimationCurvePayloadData& CurvePayload, const FString& CurveName, int32 CurveFlags, bool bDoNotImportCurveWithZero, bool bAddCurveMetadataToSkeleton, bool bShouldTransact)
 	{
 		constexpr bool bIsMorphTargetCurve = true;
 		constexpr bool bIsMaterialCurve = false;
-		return InternalCreateCurve(TargetSequence, CurvePayload, CurveName, CurveFlags, bDoNotImportCurveWithZero, bIsMorphTargetCurve, bIsMaterialCurve, bShouldTransact);
+		return InternalCreateCurve(TargetSequence, CurvePayload, CurveName, CurveFlags, bDoNotImportCurveWithZero, bAddCurveMetadataToSkeleton, bIsMorphTargetCurve, bIsMaterialCurve, bShouldTransact);
 	}
 
-	bool CreateMaterialCurve(UAnimSequence* TargetSequence, FAnimationCurvePayloadData& CurvePayload, const FString& CurveName, int32 CurveFlags, bool bDoNotImportCurveWithZero, bool bShouldTransact)
+	bool CreateMaterialCurve(UAnimSequence* TargetSequence, FAnimationCurvePayloadData& CurvePayload, const FString& CurveName, int32 CurveFlags, bool bDoNotImportCurveWithZero, bool bAddCurveMetadataToSkeleton, bool bShouldTransact)
 	{
 		constexpr bool bIsMorphTargetCurve = false;
 		constexpr bool bIsMaterialCurve = true;
-		return InternalCreateCurve(TargetSequence, CurvePayload, CurveName, CurveFlags, bDoNotImportCurveWithZero, bIsMorphTargetCurve, bIsMaterialCurve, bShouldTransact);
+		return InternalCreateCurve(TargetSequence, CurvePayload, CurveName, CurveFlags, bDoNotImportCurveWithZero, bAddCurveMetadataToSkeleton, bIsMorphTargetCurve, bIsMaterialCurve, bShouldTransact);
 	}
 
-	bool CreateAttributeCurve(UAnimSequence* TargetSequence, FAnimationCurvePayloadData& CurvePayload, const FString& CurveName, int32 CurveFlags, bool bDoNotImportCurveWithZero, bool bShouldTransact)
+	bool CreateAttributeCurve(UAnimSequence* TargetSequence, FAnimationCurvePayloadData& CurvePayload, const FString& CurveName, int32 CurveFlags, bool bDoNotImportCurveWithZero, bool bAddCurveMetadataToSkeleton, bool bShouldTransact)
 	{
 		//This curve don't animate morph target or material parameter.
 		constexpr bool bIsMorphTargetCurve = false;
 		constexpr bool bIsMaterialCurve = false;
-		return InternalCreateCurve(TargetSequence, CurvePayload, CurveName, CurveFlags, bDoNotImportCurveWithZero, bIsMorphTargetCurve, bIsMaterialCurve, bShouldTransact);
+		return InternalCreateCurve(TargetSequence, CurvePayload, CurveName, CurveFlags, bDoNotImportCurveWithZero, bAddCurveMetadataToSkeleton, bIsMorphTargetCurve, bIsMaterialCurve, bShouldTransact);
 	}
 
 	void RetrieveAnimationPayloads(UAnimSequence* AnimSequence
@@ -449,21 +440,21 @@ namespace UE::Interchange::Private
 		AnimSequenceFactoryNode->GetCustomDeleteExistingNonCurveCustomAttributes(bDeleteExistingNonCurveCustomAttributes);
 		if (bDeleteExistingMorphTargetCurves || bDeleteExistingCustomAttributeCurves)
 		{
-			TArray<FSmartName> CurveSmartNamesToRemove;
+			TArray<FName> CurveNamesToRemove;
 			for (const FFloatCurve& Curve : AnimSequence->GetDataModel()->GetFloatCurves())
 			{
-				const FCurveMetaData* MetaData = Skeleton->GetCurveMetaData(Curve.Name);
+				const FCurveMetaData* MetaData = Skeleton->GetCurveMetaData(Curve.GetName());
 				if (MetaData)
 				{
 					bool bDeleteCurve = MetaData->Type.bMorphtarget ? bDeleteExistingMorphTargetCurves : bDeleteExistingCustomAttributeCurves;
 					if (bDeleteCurve)
 					{
-						CurveSmartNamesToRemove.Add(Curve.Name);
+						CurveNamesToRemove.Add(Curve.GetName());
 					}
 				}
 			}
 
-			for (const FSmartName& CurveName : CurveSmartNamesToRemove)
+			for (const FName& CurveName : CurveNamesToRemove)
 			{
 				const FAnimationCurveIdentifier CurveId(CurveName, ERawCurveTrackTypes::RCT_Float);
 				Controller.RemoveCurve(CurveId, bShouldTransact);
@@ -488,11 +479,11 @@ namespace UE::Interchange::Private
 
 			for (const FFloatCurve& FloatCurve : CurveData.FloatCurves)
 			{
-				const FCurveMetaData* MetaData = Skeleton->GetCurveMetaData(FloatCurve.Name);
+				const FCurveMetaData* MetaData = Skeleton->GetCurveMetaData(FloatCurve.GetName());
 
 				if (MetaData && !MetaData->Type.bMorphtarget)
 				{
-					OutCurvesNotFound.Add(FloatCurve.Name.DisplayName.ToString());
+					OutCurvesNotFound.Add(FloatCurve.GetName().ToString());
 				}
 			}
 
@@ -502,6 +493,8 @@ namespace UE::Interchange::Private
 			AnimSequenceFactoryNode->GetAnimatedMaterialCurveSuffixes(MaterialSuffixes);
 			bool bDoNotImportCurveWithZero = false;
 			AnimSequenceFactoryNode->GetCustomDoNotImportCurveWithZero(bDoNotImportCurveWithZero);
+			bool bAddCurveMetadataToSkeleton = false;
+			AnimSequenceFactoryNode->GetCustomAddCurveMetadataToSkeleton(bAddCurveMetadataToSkeleton);
 			bool bRemoveCurveRedundantKeys = false;
 			AnimSequenceFactoryNode->GetCustomRemoveCurveRedundantKeys(bRemoveCurveRedundantKeys);
 
@@ -577,7 +570,7 @@ namespace UE::Interchange::Private
 						}
 					}
 					constexpr int32 CurveFlags = 0;
-					CreateMorphTargetCurve(AnimSequence, CurvePayload, AnimationCurveMorphTargetNodeNames.FindChecked(CurveNameAndPayload.Key), CurveFlags, bDoNotImportCurveWithZero, bShouldTransact);
+					CreateMorphTargetCurve(AnimSequence, CurvePayload, AnimationCurveMorphTargetNodeNames.FindChecked(CurveNameAndPayload.Key), CurveFlags, bDoNotImportCurveWithZero, bAddCurveMetadataToSkeleton, bShouldTransact);
 				}
 			}
 
@@ -650,11 +643,11 @@ namespace UE::Interchange::Private
 						constexpr int32 CurveFlags = 0;
 						if (bMaterialDriveParameterOnCustomAttribute || IsCurveHookToMaterial(CurveName))
 						{
-							CreateMaterialCurve(AnimSequence, CurvePayload, CurveName, CurveFlags, bDoNotImportCurveWithZero, bShouldTransact);
+							CreateMaterialCurve(AnimSequence, CurvePayload, CurveName, CurveFlags, bDoNotImportCurveWithZero, bAddCurveMetadataToSkeleton, bShouldTransact);
 						}
 						else
 						{
-							CreateAttributeCurve(AnimSequence, CurvePayload, CurveName, CurveFlags, bDoNotImportCurveWithZero, bShouldTransact);
+							CreateAttributeCurve(AnimSequence, CurvePayload, CurveName, CurveFlags, bDoNotImportCurveWithZero, bAddCurveMetadataToSkeleton, bShouldTransact);
 						}
 					}
 				}
@@ -749,10 +742,10 @@ namespace UE::Interchange::Private
 			// Store float curve tracks which use to exist on the animation
 			for (const FFloatCurve& Curve : AnimSequence->GetDataModel()->GetFloatCurves())
 			{
-				const FCurveMetaData* MetaData = Skeleton->GetCurveMetaData(Curve.Name);
+				const FCurveMetaData* MetaData = Skeleton->GetCurveMetaData(Curve.GetName());
 				if (MetaData && !MetaData->Type.bMorphtarget)
 				{
-					OutCurvesNotFound.Add(Curve.Name.DisplayName.ToString());
+					OutCurvesNotFound.Add(Curve.GetName().ToString());
 				}
 			}
 		}
diff --git a/Engine/Plugins/Interchange/Runtime/Source/Pipelines/Public/InterchangeGenericAnimationPipeline.h b/Engine/Plugins/Interchange/Runtime/Source/Pipelines/Public/InterchangeGenericAnimationPipeline.h
--- a/Engine/Plugins/Interchange/Runtime/Source/Pipelines/Public/InterchangeGenericAnimationPipeline.h
+++ b/Engine/Plugins/Interchange/Runtime/Source/Pipelines/Public/InterchangeGenericAnimationPipeline.h
@@ -74,6 +74,10 @@ public:
 	UPROPERTY(EditAnywhere, BlueprintReadWrite, Category = "Animations", meta = (DisplayName = "Import Attributes as Curves or Animation Attributes"))
 	bool bImportCustomAttribute = true;
 
+	/** Whether to automatically add curve metadata to an animation's skeleton. If this is disabled, curve metadata will be added to skeletal meshes for morph targets, but no metadata entry will be created for general curves. */
+	UPROPERTY(EditAnywhere, BlueprintReadWrite, Category = "Animations", meta = (EditCondition = "bImportCustomAttribute"))
+	bool bAddCurveMetadataToSkeleton = true;
+
 	/** Set Material Curve Type for all custom attributes that exists */
 	UPROPERTY(EditAnywhere, BlueprintReadWrite, Category = "Animations", meta = (EditCondition = "bImportCustomAttribute", DisplayName = "Set Material Curve Type"))
 	bool bSetMaterialDriveParameterOnCustomAttribute = false;
diff --git a/Engine/Plugins/Runtime/SkeletalMerging/Source/SkeletalMerging/Private/SkeletalMergingLibrary.cpp b/Engine/Plugins/Runtime/SkeletalMerging/Source/SkeletalMerging/Private/SkeletalMergingLibrary.cpp
--- a/Engine/Plugins/Runtime/SkeletalMerging/Source/SkeletalMerging/Private/SkeletalMergingLibrary.cpp
+++ b/Engine/Plugins/Runtime/SkeletalMerging/Source/SkeletalMerging/Private/SkeletalMergingLibrary.cpp
@@ -222,15 +222,10 @@ USkeleton* USkeletalMergingLibrary::MergeSkeletons(const FSkeletonMergeParams& P
 
 		if (Params.bMergeCurveNames)
 		{
-			if (const FSmartNameMapping* CurveMappingPtr = Skeleton->GetSmartNameContainer(USkeleton::AnimCurveMappingName))
+			Skeleton->ForEachCurveMetaData([&UniqueCurveNames](FName InCurveName, const FCurveMetaData& InCurveMetaData)
 			{
-				TArray<FName> CurveNames;
-				CurveMappingPtr->FillNameArray(CurveNames);
-				for (const FName& CurveName : CurveNames)
-				{
-					UniqueCurveNames.FindOrAdd(CurveName) = CurveMappingPtr->GetCurveMetaData(CurveName);
-				}
-			}
+				UniqueCurveNames.FindOrAdd(InCurveName, &InCurveMetaData);
+			});
 		}
 
 		if (Params.bMergeAnimSlotGroups)
@@ -333,14 +328,6 @@ void USkeletalMergingLibrary::AddVirtualBones(USkeleton* InSkeleton, const TArra
 
 void USkeletalMergingLibrary::AddCurveNames(USkeleton* InSkeleton, const TMap<FName, const FCurveMetaData*>& InCurves)
 {
-	TArray<FSmartName> CurveSmartNames;
-
-	Algo::Transform(InCurves, CurveSmartNames, [](const TPair<FName, const FCurveMetaData*>& CurveMetaDataPair)
-	{
-		return FSmartName(CurveMetaDataPair.Key, INDEX_NONE);
-	});
-	InSkeleton->VerifySmartNames(USkeleton::AnimCurveMappingName, CurveSmartNames);
-		
 	for(const TPair<FName, const FCurveMetaData*>& CurveMetaDataPair : InCurves)
 	{
 		if (CurveMetaDataPair.Value)
diff --git a/Engine/Plugins/Tests/EditorTests/Source/EditorTests/Private/UnrealEd/FbxAutomationTests.cpp b/Engine/Plugins/Tests/EditorTests/Source/EditorTests/Private/UnrealEd/FbxAutomationTests.cpp
--- a/Engine/Plugins/Tests/EditorTests/Source/EditorTests/Private/UnrealEd/FbxAutomationTests.cpp
+++ b/Engine/Plugins/Tests/EditorTests/Source/EditorTests/Private/UnrealEd/FbxAutomationTests.cpp
@@ -96,12 +96,7 @@ public:
 			}
 
 			const FString& CurveName = ExpectedResult.ExpectedPresetsDataString[0];
-			const FFloatCurve* FloatCurve = nullptr;
-			FSmartName OutSmartName;
-			if (AnimSequence->GetSkeleton()->GetSmartNameByName(USkeleton::AnimCurveMappingName, *CurveName, OutSmartName))
-			{
-				FloatCurve = static_cast<const FFloatCurve*>(AnimSequence->GetCurveData().GetCurveData(OutSmartName.UID, ERawCurveTrackTypes::RCT_Float));
-			}
+			const FFloatCurve* FloatCurve = static_cast<const FFloatCurve*>(AnimSequence->GetCurveData().GetCurveData(FName(*CurveName), ERawCurveTrackTypes::RCT_Float));
 			if (FloatCurve == nullptr)
 			{
 				ExecutionInfo.AddError(FString::Printf(TEXT("%s no custom curve named %s was imported"),
diff --git a/Engine/Plugins/Tests/InterchangeTests/Source/InterchangeTests/Private/ImportTestFunctions/AnimationImportTestFunctions.cpp b/Engine/Plugins/Tests/InterchangeTests/Source/InterchangeTests/Private/ImportTestFunctions/AnimationImportTestFunctions.cpp
--- a/Engine/Plugins/Tests/InterchangeTests/Source/InterchangeTests/Private/ImportTestFunctions/AnimationImportTestFunctions.cpp
+++ b/Engine/Plugins/Tests/InterchangeTests/Source/InterchangeTests/Private/ImportTestFunctions/AnimationImportTestFunctions.cpp
@@ -12,18 +12,16 @@ namespace UE::Private::AnimationImportTestFunction
 	static bool GetImportedCustomCurveKey(UAnimSequence* AnimSequence, const FString& CurveName, const int32 KeyIndex, FRichCurveKey& OutCurveKey, FInterchangeTestFunctionResult& Result)
 	{
 		const FFloatCurve* FloatCurve = nullptr;
-		FSmartName OutSmartName;
-		if (AnimSequence->GetSkeleton()->GetSmartNameByName(USkeleton::AnimCurveMappingName, *CurveName, OutSmartName))
+
+		for (const FFloatCurve& Curve : AnimSequence->GetCurveData().FloatCurves)
 		{
-			for (const FFloatCurve& Curve : AnimSequence->GetCurveData().FloatCurves)
+			if (Curve.GetName() == FName(*CurveName))
 			{
-				if (Curve.Name.DisplayName == FName(*CurveName))
-				{
-					FloatCurve = &Curve;
-					break;
-				}
+				FloatCurve = &Curve;
+				break;
 			}
 		}
+
 		if (FloatCurve == nullptr)
 		{
 			Result.AddError(FString::Printf(TEXT("No custom curve named %s was imported"), *CurveName));
@@ -168,4 +166,3 @@ FInterchangeTestFunctionResult UAnimationImportTestFunctions::CheckCurveKeyLeave
 	}
 	return Result;
 }
-
diff --git a/Engine/Source/Developer/AnimationDataController/Private/AnimDataController.cpp b/Engine/Source/Developer/AnimationDataController/Private/AnimDataController.cpp
--- a/Engine/Source/Developer/AnimationDataController/Private/AnimDataController.cpp
+++ b/Engine/Source/Developer/AnimationDataController/Private/AnimDataController.cpp
@@ -253,128 +253,6 @@ void UAnimDataController::SetFrameRate(FFrameRate FrameRate, bool bShouldTransac
 	}
 }
 
-
-void UAnimDataController::UpdateCurveNamesFromSkeleton(const USkeleton* Skeleton, ERawCurveTrackTypes SupportedCurveType, bool bShouldTransact /*= true*/)
-{
-	ValidateModel();
-
-	if (Skeleton)
-	{
-		if (IsSupportedCurveType(SupportedCurveType))
-		{
-			FBracket Bracket = ConditionalBracket(LOCTEXT("ValidateRawCurves", "Validating Animation Curve Names"), bShouldTransact);
-			switch (SupportedCurveType)
-			{
-			case ERawCurveTrackTypes::RCT_Float:
-			{
-				const FSmartNameMapping* NameMapping = Skeleton->GetSmartNameContainer(USkeleton::AnimCurveMappingName);
-				for (FFloatCurve& FloatCurve : Model->CurveData.FloatCurves)
-				{
-					FSmartName NewSmartName = FloatCurve.Name;
-					NameMapping->GetName(FloatCurve.Name.UID, NewSmartName.DisplayName);
-					if (NewSmartName != FloatCurve.Name)
-					{
-						const FAnimationCurveIdentifier CurrentId(FloatCurve.Name, SupportedCurveType);
-						const FAnimationCurveIdentifier NewId(NewSmartName, SupportedCurveType);
-						RenameCurve(CurrentId, NewId, bShouldTransact);
-					}
-				}
-				break;
-			}
-			case ERawCurveTrackTypes::RCT_Transform:
-			{
-				const FSmartNameMapping* NameMapping = Skeleton->GetSmartNameContainer(USkeleton::AnimTrackCurveMappingName);
-				for (FTransformCurve& TransformCurve : Model->CurveData.TransformCurves)
-				{
-					FSmartName NewSmartName = TransformCurve.Name;
-					NameMapping->GetName(TransformCurve.Name.UID, NewSmartName.DisplayName);
-					if (NewSmartName != TransformCurve.Name)
-					{
-						const FAnimationCurveIdentifier CurrentId(TransformCurve.Name, SupportedCurveType);
-						const FAnimationCurveIdentifier NewId(NewSmartName, SupportedCurveType);
-						RenameCurve(CurrentId, NewId, bShouldTransact);
-					}
-				}
-				break;
-			}
-			case ERawCurveTrackTypes::RCT_Vector:
-			default:
-				const FString CurveTypeAsString = GetCurveTypeValueName(SupportedCurveType);
-				ReportWarningf(LOCTEXT("InvalidCurveTypeWarning", "Invalid curve type provided: {0} ({1})"), FText::FromString(CurveTypeAsString), FText::AsNumber(static_cast<int32>(SupportedCurveType)));
-				
-			}
-		}
-		else
-		{
-			const FString CurveTypeAsString = GetCurveTypeValueName(SupportedCurveType);
-			ReportWarningf(LOCTEXT("InvalidCurveTypeWarning", "Invalid curve type provided: {0} ({1})"), FText::FromString(CurveTypeAsString), FText::AsNumber(static_cast<int32>(SupportedCurveType)));
-		}
-	}
-	else
-	{
-		Report(ELogVerbosity::Error, LOCTEXT("UpdateCurveInvalidSkeletonError", "Invalid USkeleton supplied"));
-	}
-}
-
-void UAnimDataController::FindOrAddCurveNamesOnSkeleton(USkeleton* Skeleton, ERawCurveTrackTypes SupportedCurveType, bool bShouldTransact)
-{
-	ValidateModel();
-	
-	if (Skeleton)
-	{
-		if (IsSupportedCurveType(SupportedCurveType))
-		{
-			FBracket Bracket = ConditionalBracket(LOCTEXT("FindOrAddRawCurveNames", "Updating Skeleton with Animation Curve Names"), bShouldTransact);
-			switch (SupportedCurveType)
-			{
-			case ERawCurveTrackTypes::RCT_Float:
-			{
-				for (FFloatCurve& FloatCurve : Model->CurveData.FloatCurves)
-				{
-					FSmartName NewSmartName = FloatCurve.Name;
-					Skeleton->VerifySmartName(USkeleton::AnimCurveMappingName, NewSmartName);
-					if (NewSmartName != FloatCurve.Name)
-					{
-						const FAnimationCurveIdentifier CurrentId(FloatCurve.Name, SupportedCurveType);
-						const FAnimationCurveIdentifier NewId(NewSmartName, SupportedCurveType);
-						RenameCurve(CurrentId, NewId, bShouldTransact);
-					}
-				}
-				break;
-			}
-			case ERawCurveTrackTypes::RCT_Transform:
-			{
-				for (FTransformCurve& TransformCurve : Model->CurveData.TransformCurves)
-				{
-					FSmartName NewSmartName = TransformCurve.Name;
-					Skeleton->VerifySmartName(USkeleton::AnimTrackCurveMappingName, NewSmartName);
-					if (NewSmartName != TransformCurve.Name)
-					{
-						const FAnimationCurveIdentifier CurrentId(TransformCurve.Name, SupportedCurveType);
-						const FAnimationCurveIdentifier NewId(NewSmartName, SupportedCurveType);
-						RenameCurve(CurrentId, NewId, bShouldTransact);
-					}
-				}
-				break;
-			}
-			case ERawCurveTrackTypes::RCT_Vector:
-			default:
-				const FString CurveTypeAsString = GetCurveTypeValueName(SupportedCurveType);
-				ReportWarningf(LOCTEXT("InvalidCurveTypeWarning", "Invalid curve type provided: {0} ({1})"), FText::FromString(CurveTypeAsString), FText::AsNumber(static_cast<int32>(SupportedCurveType)));
-			}
-		}
-		else
-		{
-			const FString CurveTypeAsString = GetCurveTypeValueName(SupportedCurveType);
-			ReportWarningf(LOCTEXT("InvalidCurveTypeWarning", "Invalid curve type provided: {0} ({1})"), FText::FromString(CurveTypeAsString), FText::AsNumber(static_cast<int32>(SupportedCurveType)));
-		}
-	}
-	else
-	{
-		Report(ELogVerbosity::Error, LOCTEXT("FindOrAddCurveInvalidSkeletonError", "Invalid USkeleton supplied "));
-	}
-}
-
 bool UAnimDataController::RemoveBoneTracksMissingFromSkeleton(const USkeleton* Skeleton, bool bShouldTransact /*= true*/)
 {
 	if (!ModelInterface->GetAnimationSequence())
@@ -506,8 +384,8 @@ void UAnimDataController::ResetModel(bool bShouldTransact /*= true*/)
 bool UAnimDataController::AddCurve(const FAnimationCurveIdentifier& CurveId, int32 CurveFlags /*= EAnimAssetCurveFlags::AACF_Editable*/, bool bShouldTransact /*= true*/)
 {
 	ValidateModel();
-	if (CurveId.InternalName.IsValid())
-	{		
+	if (CurveId.CurveName != NAME_None)
+	{
 		if (IsSupportedCurveType(CurveId.CurveType))
 		{
 			if (!Model->FindCurve(CurveId))
@@ -517,7 +395,7 @@ bool UAnimDataController::AddCurve(const FAnimationCurveIdentifier& CurveId, int
 				FCurveAddedPayload Payload;
 				Payload.Identifier = CurveId;
 				
-				auto AddNewCurve = [CurveName = CurveId.InternalName, CurveFlags](auto& CurveTypeArray)
+				auto AddNewCurve = [CurveName = CurveId.CurveName, CurveFlags](auto& CurveTypeArray)
 				{
 					CurveTypeArray.Add({ CurveName, CurveFlags});
 				};
@@ -544,7 +422,7 @@ bool UAnimDataController::AddCurve(const FAnimationCurveIdentifier& CurveId, int
 			else
 			{
 				const FString CurveTypeAsString = GetCurveTypeValueName(CurveId.CurveType);
-				ReportWarningf(LOCTEXT("ExistingCurveNameWarning", "Curve with name {0} and type {1} ({2}) already exists"), FText::FromName(CurveId.InternalName.DisplayName), FText::FromString(CurveTypeAsString), FText::AsNumber(static_cast<int32>(CurveId.CurveType)));
+				ReportWarningf(LOCTEXT("ExistingCurveNameWarning", "Curve with name {0} and type {1} ({2}) already exists"), FText::FromName(CurveId.CurveName), FText::FromString(CurveTypeAsString), FText::AsNumber(static_cast<int32>(CurveId.CurveType)));
 			}			
 		}
 		else 
@@ -556,7 +434,7 @@ bool UAnimDataController::AddCurve(const FAnimationCurveIdentifier& CurveId, int
 	else
 	{
 		const FString CurveTypeAsString = GetCurveTypeValueName(CurveId.CurveType);
-		ReportWarningf(LOCTEXT("InvalidCurveIdentifierWarning", "Invalid curve identifier provided: name: {0}, UID: {1} type: {2}"), FText::FromName(CurveId.InternalName.DisplayName), FText::AsNumber(CurveId.InternalName.UID), FText::FromString(CurveTypeAsString));
+		ReportWarningf(LOCTEXT("InvalidCurveIdentifierWarning", "Invalid curve identifier provided: name: {0}, type: {1}"), FText::FromName(CurveId.CurveName), FText::FromString(CurveTypeAsString));
 	}
 
 	return false;
@@ -568,7 +446,7 @@ bool UAnimDataController::DuplicateCurve(const FAnimationCurveIdentifier& CopyCu
 
 	ERawCurveTrackTypes SupportedCurveType = CopyCurveId.CurveType;
 
-	if (CopyCurveId.InternalName.IsValid() && NewCurveId.InternalName.IsValid())
+	if (CopyCurveId.CurveName != NAME_None && NewCurveId.CurveName != NAME_None)
 	{
 		if (IsSupportedCurveType(SupportedCurveType))
 		{
@@ -580,7 +458,7 @@ bool UAnimDataController::DuplicateCurve(const FAnimationCurveIdentifier& CopyCu
 					{
 						FTransaction Transaction = ConditionalTransaction(LOCTEXT("CopyRawCurve", "Duplicating Animation Curve"), bShouldTransact);
 
-						auto DuplicateCurve = [NewCurveName = NewCurveId.InternalName](auto& CurveDataArray, const auto& SourceCurve)
+						auto DuplicateCurve = [NewCurveName = NewCurveId.CurveName](auto& CurveDataArray, const auto& SourceCurve)
 						{
 							auto& DuplicatedCurve = CurveDataArray.Add_GetRef( { NewCurveName, SourceCurve.GetCurveTypeFlags() });
 							DuplicatedCurve.CopyCurve(SourceCurve);
@@ -611,13 +489,13 @@ bool UAnimDataController::DuplicateCurve(const FAnimationCurveIdentifier& CopyCu
 					else
 					{
 						const FString CurveTypeAsString = GetCurveTypeValueName(NewCurveId.CurveType);
-						ReportWarningf(LOCTEXT("ExistingCurveNameWarning", "Curve with name {0} and type {1} ({2}) already exists"), FText::FromName(NewCurveId.InternalName.DisplayName), FText::FromString(CurveTypeAsString), FText::AsNumber(static_cast<int32>(NewCurveId.CurveType)));
+						ReportWarningf(LOCTEXT("ExistingCurveNameWarning", "Curve with name {0} and type {1} ({2}) already exists"), FText::FromName(NewCurveId.CurveName), FText::FromString(CurveTypeAsString), FText::AsNumber(static_cast<int32>(NewCurveId.CurveType)));
 					}
 				}
 				else
 				{
 					const FString CurveTypeAsString = GetCurveTypeValueName(CopyCurveId.CurveType);
-					ReportWarningf(LOCTEXT("CurveNameToDuplicateNotFoundWarning", "Could not find curve with name {0} and type {1} ({2}) for duplication"), FText::FromName(NewCurveId.InternalName.DisplayName), FText::FromString(CurveTypeAsString), FText::AsNumber(static_cast<int32>(NewCurveId.CurveType)));
+					ReportWarningf(LOCTEXT("CurveNameToDuplicateNotFoundWarning", "Could not find curve with name {0} and type {1} ({2}) for duplication"), FText::FromName(NewCurveId.CurveName), FText::FromString(CurveTypeAsString), FText::AsNumber(static_cast<int32>(NewCurveId.CurveType)));
 				}
 			}
 		}
@@ -627,6 +505,10 @@ bool UAnimDataController::DuplicateCurve(const FAnimationCurveIdentifier& CopyCu
 			ReportWarningf(LOCTEXT("InvalidCurveTypeWarning", "Invalid curve type provided: {0} ({1})"), FText::FromString(CurveTypeAsString), FText::AsNumber(static_cast<int32>(SupportedCurveType)));
 		}
 	}
+	else
+	{
+		ReportWarning(LOCTEXT("InvalidCurveIdentifierWarning", "Invalid curve identifier provided"));
+	}
 
 	return false;
 }
@@ -637,7 +519,7 @@ bool UAnimDataController::RemoveCurve(const FAnimationCurveIdentifier& CurveId,
 	ValidateModel();
 	ERawCurveTrackTypes SupportedCurveType = CurveId.CurveType;
 
-	if (CurveId.InternalName.IsValid())
+	if (CurveId.CurveName != NAME_None)
 	{
 		if (IsSupportedCurveType(CurveId.CurveType))
 		{
@@ -651,14 +533,14 @@ bool UAnimDataController::RemoveCurve(const FAnimationCurveIdentifier& CurveId,
 					{
 						const FTransformCurve& TransformCurve = Model->GetTransformCurve(CurveId);
 						ConditionalAction<UE::Anim::FAddTransformCurveAction>(bShouldTransact, CurveId, TransformCurve.GetCurveTypeFlags(), TransformCurve);
-						Model->CurveData.TransformCurves.RemoveAll([Name = TransformCurve.Name](const FTransformCurve& ToRemoveCurve) { return ToRemoveCurve.Name == Name; });
+						Model->CurveData.TransformCurves.RemoveAll([Name = TransformCurve.GetName()](const FTransformCurve& ToRemoveCurve) { return ToRemoveCurve.GetName() == Name; });
 						break;
 					}
 					case ERawCurveTrackTypes::RCT_Float:
 					{
 						const FFloatCurve& FloatCurve = Model->GetFloatCurve(CurveId);
 						ConditionalAction<UE::Anim::FAddFloatCurveAction>(bShouldTransact, CurveId, FloatCurve.GetCurveTypeFlags(), FloatCurve.FloatCurve.GetConstRefOfKeys(), FloatCurve.Color);
-						Model->CurveData.FloatCurves.RemoveAll([Name = FloatCurve.Name](const FFloatCurve& ToRemoveCurve) { return ToRemoveCurve.Name == Name; });
+						Model->CurveData.FloatCurves.RemoveAll([Name = FloatCurve.GetName()](const FFloatCurve& ToRemoveCurve) { return ToRemoveCurve.GetName() == Name; });
 						break;
 					}
 					case ERawCurveTrackTypes::RCT_Vector:
@@ -676,7 +558,7 @@ bool UAnimDataController::RemoveCurve(const FAnimationCurveIdentifier& CurveId,
 			else
 			{
 				const FString CurveTypeAsString = GetCurveTypeValueName(SupportedCurveType);
-				ReportWarningf(LOCTEXT("UnableToFindCurveForRemovalWarning", "Unable to find curve: {0} of type {1}"), FText::FromName(CurveId.InternalName.DisplayName), FText::FromString(CurveTypeAsString));
+				ReportWarningf(LOCTEXT("UnableToFindCurveForRemovalWarning", "Unable to find curve: {0} of type {1}"), FText::FromName(CurveId.CurveName), FText::FromString(CurveTypeAsString));
 			}
 		}
 		else
@@ -685,6 +567,11 @@ bool UAnimDataController::RemoveCurve(const FAnimationCurveIdentifier& CurveId,
 			ReportWarningf(LOCTEXT("InvalidCurveTypeWarning", "Invalid curve type provided: {0} ({1})"), FText::FromString(CurveTypeAsString), FText::AsNumber(static_cast<int32>(CurveId.CurveType)));
 		}
 	}
+	else
+	{
+		const FString CurveTypeAsString = GetCurveTypeValueName(CurveId.CurveType);
+		ReportWarningf(LOCTEXT("InvalidCurveIdentifierWarning", "Invalid curve identifier provided: name: {0}, type: {1}"), FText::FromName(CurveId.CurveName), FText::FromString(CurveTypeAsString));
+	}
 
 	return false;
 }
@@ -701,7 +588,7 @@ void UAnimDataController::RemoveAllCurvesOfType(ERawCurveTrackTypes SupportedCur
 		TArray<FTransformCurve> TransformCurves = Model->CurveData.TransformCurves;
 		for (const FTransformCurve& Curve : TransformCurves)
 		{
-			RemoveCurve(FAnimationCurveIdentifier(Curve.Name, ERawCurveTrackTypes::RCT_Transform), bShouldTransact);
+			RemoveCurve(FAnimationCurveIdentifier(Curve.GetName(), ERawCurveTrackTypes::RCT_Transform), bShouldTransact);
 		}
 		break;
 	}
@@ -710,7 +597,7 @@ void UAnimDataController::RemoveAllCurvesOfType(ERawCurveTrackTypes SupportedCur
 		TArray<FFloatCurve> FloatCurves = Model->CurveData.FloatCurves;
 		for (const FFloatCurve& Curve : FloatCurves)
 		{
-			RemoveCurve(FAnimationCurveIdentifier(Curve.Name, ERawCurveTrackTypes::RCT_Float), bShouldTransact);
+			RemoveCurve(FAnimationCurveIdentifier(Curve.GetName(), ERawCurveTrackTypes::RCT_Float), bShouldTransact);
 		}
 		break;
 	}
@@ -765,7 +652,7 @@ bool UAnimDataController::SetCurveFlag(const FAnimationCurveIdentifier& CurveId,
 	else
 	{
 		const FString CurveTypeAsString = GetCurveTypeValueName(SupportedCurveType);
-		ReportWarningf(LOCTEXT("UnableToFindCurveWarning", "Unable to find curve: {0} of type {1}"), FText::FromName(CurveId.InternalName.DisplayName), FText::FromString(CurveTypeAsString));
+		ReportWarningf(LOCTEXT("UnableToFindCurveWarning", "Unable to find curve: {0} of type {1}"), FText::FromName(CurveId.CurveName), FText::FromString(CurveTypeAsString));
 	}
 
 	return false;
@@ -813,7 +700,7 @@ bool UAnimDataController::SetCurveFlags(const FAnimationCurveIdentifier& CurveId
 	else
 	{
 		const FString CurveTypeAsString = GetCurveTypeValueName(SupportedCurveType);
-		ReportWarningf(LOCTEXT("UnableToFindCurveForRemovalWarning", "Unable to find curve: {0} of type {1}"), FText::FromName(CurveId.InternalName.DisplayName), FText::FromString(CurveTypeAsString));
+		ReportWarningf(LOCTEXT("UnableToFindCurveForRemovalWarning", "Unable to find curve: {0} of type {1}"), FText::FromName(CurveId.CurveName), FText::FromString(CurveTypeAsString));
 	}
 
 	return false;
@@ -888,7 +775,7 @@ bool UAnimDataController::SetTransformCurveKeys(const FAnimationCurveIdentifier&
 		}
 		else
 		{
-			ReportWarningf(LOCTEXT("UnableToFindTransformCurveWarning", "Unable to find transform curve: {0}"), FText::FromName(CurveId.InternalName.DisplayName));
+			ReportWarningf(LOCTEXT("UnableToFindTransformCurveWarning", "Unable to find transform curve: {0}"), FText::FromName(CurveId.CurveName));
 		}
 	}
 	else
@@ -949,7 +836,7 @@ bool UAnimDataController::SetTransformCurveKey(const FAnimationCurveIdentifier&
 	}
 	else
 	{
-		ReportWarningf(LOCTEXT("UnableToFindTransformCurveWarning", "Unable to find transform curve: {0}"), FText::FromName(CurveId.InternalName.DisplayName));
+		ReportWarningf(LOCTEXT("UnableToFindTransformCurveWarning", "Unable to find transform curve: {0}"), FText::FromName(CurveId.CurveName));
 	}
 
 	return false;
@@ -962,7 +849,7 @@ bool UAnimDataController::RemoveTransformCurveKey(const FAnimationCurveIdentifie
 
 	if (Model->FindMutableTransformCurveById(CurveId))
 	{
-		const FString BaseCurveName = CurveId.InternalName.DisplayName.ToString();
+		const FString BaseCurveName = CurveId.CurveName.ToString();
 		const TArray<FString> SubCurveNames = { TEXT( "Translation"), TEXT( "Rotation"), TEXT( "Scale") };
 		const TArray<FString> ChannelCurveNames = { TEXT("X"), TEXT("Y"), TEXT("Z") };
 
@@ -984,7 +871,7 @@ bool UAnimDataController::RemoveTransformCurveKey(const FAnimationCurveIdentifie
 	}
 	else
 	{
-		ReportWarningf(LOCTEXT("UnableToFindTransformCurveWarning", "Unable to find transform curve: {0}"), FText::FromName(CurveId.InternalName.DisplayName));
+		ReportWarningf(LOCTEXT("UnableToFindTransformCurveWarning", "Unable to find transform curve: {0}"), FText::FromName(CurveId.CurveName));
 	}
 
 	return false;
@@ -1005,9 +892,9 @@ bool UAnimDataController::RenameCurve(const FAnimationCurveIdentifier& CurveToRe
 					FTransaction Transaction = ConditionalTransaction(LOCTEXT("RenameCurve", "Renaming Curve"), bShouldTransact);
 
 					FCurveRenamedPayload Payload;
-					Payload.Identifier = FAnimationCurveIdentifier(Curve->Name, CurveToRenameId.CurveType);
+					Payload.Identifier = FAnimationCurveIdentifier(Curve->GetName(), CurveToRenameId.CurveType);
 
-					Curve->Name = NewCurveId.InternalName;
+					Curve->SetName(NewCurveId.CurveName);
 					Payload.NewIdentifier = NewCurveId;
 
 					ConditionalAction<UE::Anim::FRenameCurveAction>(bShouldTransact, NewCurveId, CurveToRenameId);
@@ -1019,26 +906,26 @@ bool UAnimDataController::RenameCurve(const FAnimationCurveIdentifier& CurveToRe
 				else
 				{
 					const FString CurveTypeAsString = GetCurveTypeValueName(CurveToRenameId.CurveType);
-					ReportWarningf(LOCTEXT("UnableToFindCurveWarning", "Unable to find curve: {0} of type {1}"), FText::FromName(CurveToRenameId.InternalName.DisplayName), FText::FromString(CurveTypeAsString));
+					ReportWarningf(LOCTEXT("UnableToFindCurveWarning", "Unable to find curve: {0} of type {1}"), FText::FromName(CurveToRenameId.CurveName), FText::FromString(CurveTypeAsString));
 				}
 			}
 			else
 			{
 				const FString CurrentCurveTypeAsString = GetCurveTypeValueName(CurveToRenameId.CurveType);
 				const FString NewCurveTypeAsString = GetCurveTypeValueName(NewCurveId.CurveType);
-				ReportWarningf(LOCTEXT("MismatchOfCurveTypesWarning", "Different curve types provided between current and new curve names: {0} ({1}) and {2} ({3})"), FText::FromName(CurveToRenameId.InternalName.DisplayName), FText::FromString(CurrentCurveTypeAsString),
-					FText::FromName(NewCurveId.InternalName.DisplayName), FText::FromString(NewCurveTypeAsString));
+				ReportWarningf(LOCTEXT("MismatchOfCurveTypesWarning", "Different curve types provided between current and new curve names: {0} ({1}) and {2} ({3})"), FText::FromName(CurveToRenameId.CurveName), FText::FromString(CurrentCurveTypeAsString),
+					FText::FromName(NewCurveId.CurveName), FText::FromString(NewCurveTypeAsString));
 			}
 		}
 		else
 		{
-			ReportWarningf(LOCTEXT("MatchingCurveNamesWarning", "Provided curve names are the same: {0}"), FText::FromName(CurveToRenameId.InternalName.DisplayName));
+			ReportWarningf(LOCTEXT("MatchingCurveNamesWarning", "Provided curve names are the same: {0}"), FText::FromName(CurveToRenameId.CurveName));
 		}
 		
 	}
 	else
 	{
-		ReportWarningf(LOCTEXT("InvalidCurveIdentiferProvidedWarning", "Invalid new curve identifier provided: {2} ({3})"), FText::FromName(NewCurveId.InternalName.DisplayName), FText::AsNumber(NewCurveId.InternalName.UID));
+		ReportWarningf(LOCTEXT("InvalidCurveIdentiferProvidedWarning", "Invalid new curve identifier provided: {0}"), FText::FromName(NewCurveId.CurveName));
 	}
 
 	return false;
@@ -1069,7 +956,7 @@ bool UAnimDataController::SetCurveColor(const FAnimationCurveIdentifier& CurveId
 			else
 			{
 				const FString CurveTypeAsString = GetCurveTypeValueName(CurveId.CurveType);
-				ReportWarningf(LOCTEXT("UnableToFindCurveWarning", "Unable to find curve: {0} of type {1}"), FText::FromName(CurveId.InternalName.DisplayName), FText::FromString(CurveTypeAsString));
+				ReportWarningf(LOCTEXT("UnableToFindCurveWarning", "Unable to find curve: {0} of type {1}"), FText::FromName(CurveId.CurveName), FText::FromString(CurveTypeAsString));
 			}
 		}
 		else
@@ -1079,7 +966,7 @@ bool UAnimDataController::SetCurveColor(const FAnimationCurveIdentifier& CurveId
 	}
 	else
 	{
-		ReportWarningf(LOCTEXT("InvalidCurveIdentifier", "Invalid Curve Identifier : {0} ({1})"), FText::FromName(CurveId.InternalName.DisplayName), FText::AsNumber(CurveId.InternalName.UID));
+		ReportWarningf(LOCTEXT("InvalidCurveIdentifier", "Invalid Curve Identifier : {0}"), FText::FromName(CurveId.CurveName));
 	}	
 
 	return false;
@@ -1111,7 +998,7 @@ bool UAnimDataController::ScaleCurve(const FAnimationCurveIdentifier& CurveId, f
 		}
 		else
 		{
-			ReportWarningf(LOCTEXT("UnableToFindFloatCurveWarning", "Unable to find float curve: {0}"), FText::FromName(CurveId.InternalName.DisplayName));
+			ReportWarningf(LOCTEXT("UnableToFindFloatCurveWarning", "Unable to find float curve: {0}"), FText::FromName(CurveId.CurveName));
 		}
 	}
 	else
@@ -1189,7 +1076,7 @@ bool UAnimDataController::RemoveCurveKey(const FAnimationCurveIdentifier& CurveI
 		}
 		else
 		{
-			ReportErrorf(LOCTEXT("RichCurveKeyNotFoundError", "Unable to find rich curve key: curve name {0}, time {1}"), FText::FromName(CurveId.InternalName.DisplayName), FText::AsNumber(Time));
+			ReportErrorf(LOCTEXT("RichCurveKeyNotFoundError", "Unable to find rich curve key: curve name {0}, time {1}"), FText::FromName(CurveId.CurveName), FText::AsNumber(Time));
 		}
 	}
 
@@ -1734,7 +1621,7 @@ void UAnimDataController::ResizeCurves(float NewLength, bool bInserted, float T0
 	{
 		FFloatCurve ResizedCurve = Curve;
 		ResizedCurve.Resize(NewLength, bInserted, T0, T1);
-		SetCurveKeys(FAnimationCurveIdentifier(Curve.Name, ERawCurveTrackTypes::RCT_Float), ResizedCurve.FloatCurve.GetConstRefOfKeys(), bShouldTransact);
+		SetCurveKeys(FAnimationCurveIdentifier(Curve.GetName(), ERawCurveTrackTypes::RCT_Float), ResizedCurve.FloatCurve.GetConstRefOfKeys(), bShouldTransact);
 	}
 
 	for (FTransformCurve& Curve : Model->CurveData.TransformCurves)
@@ -1747,7 +1634,7 @@ void UAnimDataController::ResizeCurves(float NewLength, bool bInserted, float T0
 			for (int32 ChannelIndex = 0; ChannelIndex < 3; ++ChannelIndex)
 			{
 				const EVectorCurveChannel Axis = static_cast<EVectorCurveChannel>(ChannelIndex);
-				FAnimationCurveIdentifier TargetCurveIdentifier = FAnimationCurveIdentifier(Curve.Name, ERawCurveTrackTypes::RCT_Transform);
+				FAnimationCurveIdentifier TargetCurveIdentifier = FAnimationCurveIdentifier(Curve.GetName(), ERawCurveTrackTypes::RCT_Transform);
 				UAnimationCurveIdentifierExtensions::GetTransformChildCurveIdentifier(TargetCurveIdentifier, Channel, Axis);
 				
 				FRichCurve& ChannelCurve = SubCurve.FloatCurves[ChannelIndex];
@@ -2141,13 +2028,6 @@ void UAnimDataController::PopulateWithExistingModel(TScriptInterface<IAnimationD
 	Model->NumberOfKeys = InModel->GetNumberOfFrames() + 1;
 	Model->CurveData = InModel->GetCurveData();
 	Model->AnimatedBoneAttributes = InModel->GetAttributes();
-
-	// Ensure curve name UIDs are validated against the outer skeleton as they are not copied over through FSmartName
-	if (const UAnimSequenceBase* AnimSequenceBase = Cast<UAnimSequenceBase>(InModel.GetObject()->GetOuter()))
-	{
-		FindOrAddCurveNamesOnSkeleton(AnimSequenceBase->GetSkeleton(), ERawCurveTrackTypes::RCT_Float);
-		FindOrAddCurveNamesOnSkeleton(AnimSequenceBase->GetSkeleton(), ERawCurveTrackTypes::RCT_Transform);
-	}
 }
 
 #endif // WITH_EDITOR
diff --git a/Engine/Source/Developer/AnimationDataController/Private/AnimDataControllerActions.cpp b/Engine/Source/Developer/AnimationDataController/Private/AnimDataControllerActions.cpp
--- a/Engine/Source/Developer/AnimationDataController/Private/AnimDataControllerActions.cpp
+++ b/Engine/Source/Developer/AnimationDataController/Private/AnimDataControllerActions.cpp
@@ -197,7 +197,7 @@ FString FAddCurveAction::ToStringInternal() const
 	const FString FloatLabel(TEXT("float"));
 	const FString TransformLabel(TEXT("transform"));
 
-	return FText::Format(LOCTEXT("AddCurveAction_Description", "Adding {0} curve '{1}'."), FText::FromString(CurveId.CurveType == ERawCurveTrackTypes::RCT_Float ? FloatLabel : TransformLabel), FText::FromName(CurveId.InternalName.DisplayName)).ToString();
+	return FText::Format(LOCTEXT("AddCurveAction_Description", "Adding {0} curve '{1}'."), FText::FromString(CurveId.CurveType == ERawCurveTrackTypes::RCT_Float ? FloatLabel : TransformLabel), FText::FromName(CurveId.CurveName)).ToString();
 }
 
 TUniquePtr<FChange> FRemoveCurveAction::ExecuteInternal(IAnimationDataModel* Model, IAnimationDataController* Controller)
@@ -225,7 +225,7 @@ FString FRemoveCurveAction::ToStringInternal() const
 	const FString FloatLabel(TEXT("float"));
 	const FString TransformLabel(TEXT("transform"));
 
-	return FText::Format(LOCTEXT("RemoveCurveAction_Description", "Removing {0} curve '{1}'."), FText::FromString(CurveId.CurveType == ERawCurveTrackTypes::RCT_Float ? FloatLabel : TransformLabel), FText::FromName(CurveId.InternalName.DisplayName)).ToString();
+	return FText::Format(LOCTEXT("RemoveCurveAction_Description", "Removing {0} curve '{1}'."), FText::FromString(CurveId.CurveType == ERawCurveTrackTypes::RCT_Float ? FloatLabel : TransformLabel), FText::FromName(CurveId.CurveName)).ToString();
 }
 
 TUniquePtr<FChange> FSetCurveFlagsAction::ExecuteInternal(IAnimationDataModel* Model, IAnimationDataController* Controller)
@@ -245,7 +245,7 @@ FString FSetCurveFlagsAction::ToStringInternal() const
 {
 	const FString FloatLabel(TEXT("float"));
 	const FString TransformLabel(TEXT("transform"));
-	return FText::Format(LOCTEXT("SetCurveFlagsAction_Description", "Setting flags for {0} curve '{1}'."), FText::FromString(CurveType == ERawCurveTrackTypes::RCT_Float ? FloatLabel : TransformLabel), FText::FromName(CurveId.InternalName.DisplayName)).ToString();
+	return FText::Format(LOCTEXT("SetCurveFlagsAction_Description", "Setting flags for {0} curve '{1}'."), FText::FromString(CurveType == ERawCurveTrackTypes::RCT_Float ? FloatLabel : TransformLabel), FText::FromName(CurveId.CurveName)).ToString();
 }
 
 TUniquePtr<FChange> FRenameCurveAction::ExecuteInternal(IAnimationDataModel* Model, IAnimationDataController* Controller)
@@ -258,7 +258,7 @@ FString FRenameCurveAction::ToStringInternal() const
 {
 	const FString FloatLabel(TEXT("float"));
 	const FString TransformLabel(TEXT("transform"));
-	return FText::Format(LOCTEXT("RenameCurveAction_Description", "Renaming {0} curve '{1}' to '{2}'."), FText::FromString(CurveId.CurveType == ERawCurveTrackTypes::RCT_Float ? FloatLabel : TransformLabel), FText::FromName(CurveId.InternalName.DisplayName), FText::FromName(NewCurveId.InternalName.DisplayName)).ToString();
+	return FText::Format(LOCTEXT("RenameCurveAction_Description", "Renaming {0} curve '{1}' to '{2}'."), FText::FromString(CurveId.CurveType == ERawCurveTrackTypes::RCT_Float ? FloatLabel : TransformLabel), FText::FromName(CurveId.CurveName), FText::FromName(NewCurveId.CurveName)).ToString();
 }
 
 FScaleCurveAction::FScaleCurveAction(const FAnimationCurveIdentifier& InCurveId, float InOrigin, float InFactor, ERawCurveTrackTypes InCurveType) : CurveId(InCurveId), CurveType(InCurveType), Origin(InOrigin), Factor(InFactor)
@@ -282,7 +282,7 @@ FString FScaleCurveAction::ToStringInternal() const
 	const FString FloatLabel(TEXT("float"));
 	const FString TransformLabel(TEXT("transform"));
 
-	return FText::Format(LOCTEXT("ScaleCurveAction_Description", "Scaling {0} curve '{1}'."), FText::FromString(CurveType == ERawCurveTrackTypes::RCT_Float ? FloatLabel : TransformLabel), FText::FromName(CurveId.InternalName.DisplayName)).ToString();
+	return FText::Format(LOCTEXT("ScaleCurveAction_Description", "Scaling {0} curve '{1}'."), FText::FromString(CurveType == ERawCurveTrackTypes::RCT_Float ? FloatLabel : TransformLabel), FText::FromName(CurveId.CurveName)).ToString();
 }
 
 TUniquePtr<FChange> FAddFloatCurveAction::ExecuteInternal(IAnimationDataModel* Model, IAnimationDataController* Controller)
@@ -300,7 +300,7 @@ TUniquePtr<FChange> FAddFloatCurveAction::ExecuteInternal(IAnimationDataModel* M
 
 FString FAddFloatCurveAction::ToStringInternal() const
 {
-	return FText::Format(LOCTEXT("AddFloatCurveAction_Format", "Adding float curve '{0}'."), FText::FromName(CurveId.InternalName.DisplayName)).ToString();
+	return FText::Format(LOCTEXT("AddFloatCurveAction_Format", "Adding float curve '{0}'."), FText::FromName(CurveId.CurveName)).ToString();
 }
 
 FAddTransformCurveAction::FAddTransformCurveAction(const FAnimationCurveIdentifier& InCurveId, int32 InFlags, const FTransformCurve& InTransformCurve) : CurveId(InCurveId), Flags(InFlags)
@@ -341,7 +341,7 @@ TUniquePtr<FChange> FAddTransformCurveAction::ExecuteInternal(IAnimationDataMode
 
 FString FAddTransformCurveAction::ToStringInternal() const
 {
-	return FText::Format(LOCTEXT("AddTransformCurveAction_Description", "Adding transform curve '{0}'."), FText::FromName(CurveId.InternalName.DisplayName)).ToString();
+	return FText::Format(LOCTEXT("AddTransformCurveAction_Description", "Adding transform curve '{0}'."), FText::FromName(CurveId.CurveName)).ToString();
 }
 
 TUniquePtr<FChange> FAddRichCurveKeyAction::ExecuteInternal(IAnimationDataModel* Model, IAnimationDataController* Controller)
@@ -356,7 +356,7 @@ FString FAddRichCurveKeyAction::ToStringInternal() const
 	const FString FloatLabel(TEXT("float"));
 	const FString TransformLabel(TEXT("transform"));
 
-	return FText::Format(LOCTEXT("AddNamedRichCurveKeyAction_Description", "Adding key to {0} curve '{1}'."), FText::FromString(CurveId.CurveType == ERawCurveTrackTypes::RCT_Float ? FloatLabel : TransformLabel), FText::FromName(CurveId.InternalName.DisplayName)).ToString();
+	return FText::Format(LOCTEXT("AddNamedRichCurveKeyAction_Description", "Adding key to {0} curve '{1}'."), FText::FromString(CurveId.CurveType == ERawCurveTrackTypes::RCT_Float ? FloatLabel : TransformLabel), FText::FromName(CurveId.CurveName)).ToString();
 }
 
 TUniquePtr<FChange> FSetRichCurveKeyAction::ExecuteInternal(IAnimationDataModel* Model, IAnimationDataController* Controller)
@@ -377,7 +377,7 @@ FString FSetRichCurveKeyAction::ToStringInternal() const
 	const FString FloatLabel(TEXT("float"));
 	const FString TransformLabel(TEXT("transform"));
 
-	return FText::Format(LOCTEXT("SetNamedRichCurveKeyAction_Description", "Setting key for {0} curve '{1}'."), FText::FromString(CurveId.CurveType == ERawCurveTrackTypes::RCT_Float ? FloatLabel : TransformLabel), FText::FromName(CurveId.InternalName.DisplayName)).ToString();
+	return FText::Format(LOCTEXT("SetNamedRichCurveKeyAction_Description", "Setting key for {0} curve '{1}'."), FText::FromString(CurveId.CurveType == ERawCurveTrackTypes::RCT_Float ? FloatLabel : TransformLabel), FText::FromName(CurveId.CurveName)).ToString();
 }
 
 TUniquePtr<FChange> FRemoveRichCurveKeyAction::ExecuteInternal(IAnimationDataModel* Model, IAnimationDataController* Controller)
@@ -398,7 +398,7 @@ FString FRemoveRichCurveKeyAction::ToStringInternal() const
 	const FString FloatLabel(TEXT("float"));
 	const FString TransformLabel(TEXT("transform"));
 
-	return FText::Format(LOCTEXT("RemoveNamedRichCurveKeyAction_Description", "Removing key from {0} curve '{1}'."), FText::FromString(CurveId.CurveType == ERawCurveTrackTypes::RCT_Float ? FloatLabel : TransformLabel), FText::FromName(CurveId.InternalName.DisplayName)).ToString();
+	return FText::Format(LOCTEXT("RemoveNamedRichCurveKeyAction_Description", "Removing key from {0} curve '{1}'."), FText::FromString(CurveId.CurveType == ERawCurveTrackTypes::RCT_Float ? FloatLabel : TransformLabel), FText::FromName(CurveId.CurveName)).ToString();
 }
 
 TUniquePtr<FChange> FSetRichCurveKeysAction::ExecuteInternal(IAnimationDataModel* Model, IAnimationDataController* Controller)
@@ -417,7 +417,7 @@ FString FSetRichCurveKeysAction::ToStringInternal() const
 	const FString FloatLabel(TEXT("float"));
 	const FString TransformLabel(TEXT("transform"));
 
-	return FText::Format(LOCTEXT("SetNamedRichCurveKeysAction_Description", "Replacing keys for {0} curve '{1}'."), FText::FromString(CurveId.CurveType == ERawCurveTrackTypes::RCT_Float ? FloatLabel : TransformLabel), FText::FromName(CurveId.InternalName.DisplayName)).ToString();
+	return FText::Format(LOCTEXT("SetNamedRichCurveKeysAction_Description", "Replacing keys for {0} curve '{1}'."), FText::FromString(CurveId.CurveType == ERawCurveTrackTypes::RCT_Float ? FloatLabel : TransformLabel), FText::FromName(CurveId.CurveName)).ToString();
 }
 
 TUniquePtr<FChange> FSetRichCurveAttributesAction::ExecuteInternal(IAnimationDataModel* Model, IAnimationDataController* Controller)
@@ -437,7 +437,7 @@ TUniquePtr<FChange> FSetRichCurveAttributesAction::ExecuteInternal(IAnimationDat
 
 FString FSetRichCurveAttributesAction::ToStringInternal() const
 {
-	return FText::Format(LOCTEXT("SetCurveAttributesAction_Description", "Setting curve attributes '{0}'."), FText::FromName(CurveId.InternalName.DisplayName)).ToString();
+	return FText::Format(LOCTEXT("SetCurveAttributesAction_Description", "Setting curve attributes '{0}'."), FText::FromName(CurveId.CurveName)).ToString();
 }
 
 TUniquePtr<FChange> FSetCurveColorAction::ExecuteInternal(IAnimationDataModel* Model, IAnimationDataController* Controller)
@@ -452,7 +452,7 @@ TUniquePtr<FChange> FSetCurveColorAction::ExecuteInternal(IAnimationDataModel* M
 
 FString FSetCurveColorAction::ToStringInternal() const
 {
-	return FText::Format(LOCTEXT("SetCurveColorAction_Description", "Setting curve color '{0}'."), FText::FromName(CurveId.InternalName.DisplayName)).ToString();
+	return FText::Format(LOCTEXT("SetCurveColorAction_Description", "Setting curve color '{0}'."), FText::FromName(CurveId.CurveName)).ToString();
 }
 
 FAddAtributeAction::FAddAtributeAction(const FAnimatedBoneAttribute& InAttribute) : AttributeId(InAttribute.Identifier)
diff --git a/Engine/Source/Developer/AnimationDataController/Public/AnimDataController.h b/Engine/Source/Developer/AnimationDataController/Public/AnimDataController.h
--- a/Engine/Source/Developer/AnimationDataController/Public/AnimDataController.h
+++ b/Engine/Source/Developer/AnimationDataController/Public/AnimDataController.h
@@ -65,8 +65,6 @@ public:
 	virtual bool RemoveCurveKey(const FAnimationCurveIdentifier& CurveId, float Time, bool bShouldTransact = true) override;
 	virtual bool SetCurveKeys(const FAnimationCurveIdentifier& CurveId, const TArray<FRichCurveKey>& CurveKeys, bool bShouldTransact = true) override;
 	virtual bool SetCurveAttributes(const FAnimationCurveIdentifier& CurveId, const FCurveAttributes& Attributes, bool bShouldTransact = true) override;
-	virtual void UpdateCurveNamesFromSkeleton(const USkeleton* Skeleton, ERawCurveTrackTypes SupportedCurveType, bool bShouldTransact = true) override;
-	virtual void FindOrAddCurveNamesOnSkeleton(USkeleton* Skeleton, ERawCurveTrackTypes SupportedCurveType, bool bShouldTransact = true) override;
 	virtual bool RemoveBoneTracksMissingFromSkeleton(const USkeleton* Skeleton, bool bShouldTransact = true) override;
 	virtual void UpdateAttributesFromSkeleton(const USkeleton* Skeleton, bool bShouldTransact = true) override;
 	virtual void NotifyPopulated() override;
diff --git a/Engine/Source/Developer/MeshBoneReduction/Private/MeshBoneReduction.cpp b/Engine/Source/Developer/MeshBoneReduction/Private/MeshBoneReduction.cpp
--- a/Engine/Source/Developer/MeshBoneReduction/Private/MeshBoneReduction.cpp
+++ b/Engine/Source/Developer/MeshBoneReduction/Private/MeshBoneReduction.cpp
@@ -321,7 +321,7 @@ public:
 					RequiredBoneIndexArray[BoneIndex] = BoneIndex;
 				}
 			}
-			FBoneContainer RequiredBones(RequiredBoneIndexArray, false, *SkeletalMesh);
+			FBoneContainer RequiredBones(RequiredBoneIndexArray, UE::Anim::ECurveFilterMode::DisallowAll, *SkeletalMesh);
 			RequiredBones.SetUseRAWData(true);
 
 			FCompactPose Pose;
diff --git a/Engine/Source/Editor/AnimGraph/Private/AnimGraphNode_ModifyCurve.cpp b/Engine/Source/Editor/AnimGraph/Private/AnimGraphNode_ModifyCurve.cpp
--- a/Engine/Source/Editor/AnimGraph/Private/AnimGraphNode_ModifyCurve.cpp
+++ b/Engine/Source/Editor/AnimGraph/Private/AnimGraphNode_ModifyCurve.cpp
@@ -7,9 +7,12 @@
 #include "Kismet2/CompilerResultsLog.h"
 #include "AnimationGraphSchema.h"
 #include "BlueprintActionDatabaseRegistrar.h"
+#include "PersonaModule.h"
 #include "Framework/Commands/UIAction.h"
 #include "ToolMenus.h"
 #include "Kismet2/BlueprintEditorUtils.h"
+#include "Framework/Application/SlateApplication.h"
+#include "Widgets/Input/STextEntryPopup.h"
 
 #define LOCTEXT_NAMESPACE "ModifyCurve"
 
@@ -33,38 +36,49 @@ FText UAnimGraphNode_ModifyCurve::GetNodeTitle(ENodeTitleType::Type TitleType) c
 	return LOCTEXT("AnimGraphNode_ModifyCurve_Title", "Modify Curve");
 }
 
-
-TArray<FName> UAnimGraphNode_ModifyCurve::GetCurvesToAdd() const
+void UAnimGraphNode_ModifyCurve::GetAddCurveMenuActions(FMenuBuilder& MenuBuilder) const
 {
-	TArray<FName> CurvesToAdd;
-
-	if(HasValidBlueprint() && GetAnimBlueprint()->TargetSkeleton)
-	{
-		const FSmartNameMapping* Mapping = GetAnimBlueprint()->TargetSkeleton->GetSmartNameContainer(USkeleton::AnimCurveMappingName);
-		if (Mapping)
+	MenuBuilder.AddMenuEntry(
+		LOCTEXT("NewCurveLabel", "New Curve..."),
+		LOCTEXT("NewCurveToolTip", "Adds a new curve to modify"),
+		FSlateIcon(),
+		FUIAction(FExecuteAction::CreateLambda([this]()
 		{
-			Mapping->FillNameArray(CurvesToAdd);
-
-			for (FName ExistingCurveName : Node.CurveNames)
-			{
-				CurvesToAdd.RemoveSingleSwap(ExistingCurveName, false);
-			}
-
-			CurvesToAdd.Sort(FNameLexicalLess());
-		}
-	}
+			FSlateApplication::Get().DismissAllMenus();
 
-	return CurvesToAdd;
-}
+			TSharedRef<STextEntryPopup> TextEntry =
+				SNew(STextEntryPopup)
+				.Label(LOCTEXT("NewCurveLabel", "New Curve Name"))
+				.OnTextCommitted_Lambda([this](FText InText, ETextCommit::Type InCommitType)
+				{
+					FSlateApplication::Get().DismissAllMenus();
+					const_cast<UAnimGraphNode_ModifyCurve*>(this)->AddCurvePin(*InText.ToString());
+				});
+
+			FSlateApplication& SlateApp = FSlateApplication::Get();
+			SlateApp.PushMenu(
+				SlateApp.GetInteractiveTopLevelWindows()[0],
+				FWidgetPath(),
+				TextEntry,
+				SlateApp.GetCursorPos(),
+				FPopupTransitionEffect::TypeInPopup
+				);
+		}))
+	);
+
+	FPersonaModule& PersonaModule = FModuleManager::LoadModuleChecked<FPersonaModule>("Persona");
+	TSharedRef<SWidget> CurvePicker = PersonaModule.CreateCurvePicker(GetAnimBlueprint()->TargetSkeleton,
+		FOnCurvePicked::CreateLambda([this](const FName& InName)
+		{
+			FSlateApplication::Get().DismissAllMenus();
+			const_cast<UAnimGraphNode_ModifyCurve*>(this)->AddCurvePin(InName);
+		}),
+		FIsCurveNameMarkedForExclusion::CreateLambda([this](const FName& InName)
+		{
+			return Node.CurveNames.Contains(InName);
+		}));
 
-void UAnimGraphNode_ModifyCurve::GetAddCurveMenuActions(FMenuBuilder& MenuBuilder) const
-{
-	TArray<FName> CurvesToAdd = GetCurvesToAdd();
-	for (FName CurveName : CurvesToAdd)
-	{
-		FUIAction Action = FUIAction(FExecuteAction::CreateUObject(const_cast<UAnimGraphNode_ModifyCurve*>(this), &UAnimGraphNode_ModifyCurve::AddCurvePin, CurveName));
-		MenuBuilder.AddMenuEntry(FText::FromName(CurveName), FText::GetEmpty(), FSlateIcon(), Action);
-	}
+	MenuBuilder.AddWidget(CurvePicker, FText::GetEmpty(), true, false);
 }
 
 void UAnimGraphNode_ModifyCurve::GetRemoveCurveMenuActions(FMenuBuilder& MenuBuilder) const
@@ -76,7 +90,6 @@ void UAnimGraphNode_ModifyCurve::GetRemoveCurveMenuActions(FMenuBuilder& MenuBui
 	}
 }
 
-
 void UAnimGraphNode_ModifyCurve::GetNodeContextMenuActions(UToolMenu* Menu, UGraphNodeContextMenuContext* Context) const
 {
 	if (!Context->bIsDebugging)
@@ -104,16 +117,12 @@ void UAnimGraphNode_ModifyCurve::GetNodeContextMenuActions(UToolMenu* Menu, UGra
 				}
 			}
 		}
-
-		// If we have more curves to add, create submenu to offer them
-		if (GetCurvesToAdd().Num() > 0)
-		{
-			Section.AddSubMenu(
-				"AddCurvePin",
-				LOCTEXT("AddCurvePin", "Add Curve Pin"),
-				LOCTEXT("AddCurvePinTooltip", "Add a new pin to drive a curve"),
-				FNewMenuDelegate::CreateUObject(this, &UAnimGraphNode_ModifyCurve::GetAddCurveMenuActions));
-		}
+		
+		Section.AddSubMenu(
+			"AddCurvePin",
+			LOCTEXT("AddCurvePin", "Add Curve Pin"),
+			LOCTEXT("AddCurvePinTooltip", "Add a new pin to drive a curve"),
+			FNewMenuDelegate::CreateUObject(this, &UAnimGraphNode_ModifyCurve::GetAddCurveMenuActions));
 
 		// If we have curves to remove, create submenu to offer them
 		if (Node.CurveNames.Num() > 0)
diff --git a/Engine/Source/Editor/AnimGraph/Private/AnimGraphNode_PoseDriver.cpp b/Engine/Source/Editor/AnimGraph/Private/AnimGraphNode_PoseDriver.cpp
--- a/Engine/Source/Editor/AnimGraph/Private/AnimGraphNode_PoseDriver.cpp
+++ b/Engine/Source/Editor/AnimGraph/Private/AnimGraphNode_PoseDriver.cpp
@@ -304,11 +304,11 @@ void UAnimGraphNode_PoseDriver::CopyTargetsFromPoseAsset()
 		Node.PoseTargets.Empty();
 
 		// For each pose we create a target
-		const TArray<FSmartName> PoseNames = PoseAsset->GetPoseNames();
+		const TArray<FName>& PoseNames = PoseAsset->GetPoseFNames();
 		for (int32 PoseIdx = 0; PoseIdx < PoseAsset->GetNumPoses(); PoseIdx++)
 		{
 			FPoseDriverTarget PoseTarget;
-			PoseTarget.DrivenName = PoseNames[PoseIdx].DisplayName;
+			PoseTarget.DrivenName = PoseNames[PoseIdx];
 
 			// Create entry for each bone
 			for (const FBoneReference& SourceBoneRef : Node.SourceBones)
diff --git a/Engine/Source/Editor/AnimGraph/Private/AnimGraphNode_TwistCorrectiveNode.cpp b/Engine/Source/Editor/AnimGraph/Private/AnimGraphNode_TwistCorrectiveNode.cpp
--- a/Engine/Source/Editor/AnimGraph/Private/AnimGraphNode_TwistCorrectiveNode.cpp
+++ b/Engine/Source/Editor/AnimGraph/Private/AnimGraphNode_TwistCorrectiveNode.cpp
@@ -30,7 +30,7 @@ FText UAnimGraphNode_TwistCorrectiveNode::GetNodeTitle(ENodeTitleType::Type Titl
 		FFormatNamedArguments Args;
 		Args.Add(TEXT("BaseBone"), FText::FromName(Node.BaseFrame.Bone.BoneName));
 		Args.Add(TEXT("TwistBone"), FText::FromName(Node.TwistFrame.Bone.BoneName));
-		Args.Add(TEXT("CurveName"), FText::FromName(Node.Curve.Name));
+		Args.Add(TEXT("CurveName"), FText::FromName(Node.CurveName));
 
 		FormattedText = FText::Format(LOCTEXT("AnimGraphNode_TwistCorrectiveNode_Title", "Twist {CurveName} = {BaseBone}:{TwistBone} "), Args);
 		return FormattedText;
@@ -44,7 +44,7 @@ FText UAnimGraphNode_TwistCorrectiveNode::GetControllerDescription() const
 
 void UAnimGraphNode_TwistCorrectiveNode::ValidateAnimNodeDuringCompilation(USkeleton* ForSkeleton, FCompilerResultsLog& MessageLog)
 {
-	if (Node.Curve.Name == NAME_None)
+	if (Node.CurveName == NAME_None)
 	{
 		MessageLog.Warning(TEXT("@@ has missing Curve Name."), this);
 	}
diff --git a/Engine/Source/Editor/AnimGraph/Private/AnimPreviewInstance.cpp b/Engine/Source/Editor/AnimGraph/Private/AnimPreviewInstance.cpp
--- a/Engine/Source/Editor/AnimGraph/Private/AnimPreviewInstance.cpp
+++ b/Engine/Source/Editor/AnimGraph/Private/AnimPreviewInstance.cpp
@@ -292,7 +292,7 @@ void FAnimPreviewInstanceProxy::RefreshCurveBoneControllers(UAnimationAsset* Ass
 			}
 
 			// add bone modifier
-			FName BoneName = Curve.Name.DisplayName;
+			FName BoneName = Curve.GetName();
 			if (BoneName != NAME_None && MySkeleton->GetReferenceSkeleton().FindBoneIndex(BoneName) != INDEX_NONE)
 			{
 				ModifyBone(BoneName, true);
diff --git a/Engine/Source/Editor/AnimGraph/Private/PoseDriverDetails.cpp b/Engine/Source/Editor/AnimGraph/Private/PoseDriverDetails.cpp
--- a/Engine/Source/Editor/AnimGraph/Private/PoseDriverDetails.cpp
+++ b/Engine/Source/Editor/AnimGraph/Private/PoseDriverDetails.cpp
@@ -1250,17 +1250,13 @@ void FPoseDriverDetails::UpdateDrivenNameOptions()
 			USkeleton* Skeleton = PoseDriver->GetAnimBlueprint()->TargetSkeleton;
 			if (Skeleton)
 			{
-				const FSmartNameMapping* Mapping = Skeleton->GetSmartNameContainer(USkeleton::AnimCurveMappingName);
-				if (Mapping)
+				TArray<FName> NameArray;
+				Skeleton->GetCurveMetaDataNames(NameArray);
+				NameArray.Sort(FNameLexicalLess());
+
+				for (const FName& CurveName : NameArray)
 				{
-					TArray<FName> NameArray;
-					Mapping->FillNameArray(NameArray);
-					NameArray.Sort(FNameLexicalLess());
-
-					for (FName CurveName : NameArray)
-					{
-						DrivenNameOptions.Add(MakeShareable(new FString(CurveName.ToString())));
-					}
+					DrivenNameOptions.Add(MakeShareable(new FString(CurveName.ToString())));
 				}
 			}
 		}
@@ -1269,10 +1265,9 @@ void FPoseDriverDetails::UpdateDrivenNameOptions()
 		{
 			if (PoseDriver->Node.PoseAsset)
 			{
-				const TArray<FSmartName> PoseNames = PoseDriver->Node.PoseAsset->GetPoseNames();
-				for (const FSmartName& SmartName : PoseNames)
+				for (const FName& Name : PoseDriver->Node.PoseAsset->GetPoseFNames())
 				{
-					DrivenNameOptions.Add(MakeShareable(new FString(SmartName.DisplayName.ToString())));
+					DrivenNameOptions.Add(MakeShareable(new FString(Name.ToString())));
 				}
 			}
 		}
diff --git a/Engine/Source/Editor/AnimGraph/Public/AnimGraphNode_ModifyCurve.h b/Engine/Source/Editor/AnimGraph/Public/AnimGraphNode_ModifyCurve.h
--- a/Engine/Source/Editor/AnimGraph/Public/AnimGraphNode_ModifyCurve.h
+++ b/Engine/Source/Editor/AnimGraph/Public/AnimGraphNode_ModifyCurve.h
@@ -45,6 +45,4 @@ private:
 	void GetAddCurveMenuActions(FMenuBuilder& MenuBuilder) const;
 	/** Create submenu with options for curves to remove */
 	void GetRemoveCurveMenuActions(FMenuBuilder& MenuBuilder) const;
-	/** Returns list of curves we have not added yet */
-	TArray<FName> GetCurvesToAdd() const;
 };
diff --git a/Engine/Source/Editor/AnimationBlueprintEditor/Private/AnimationBlueprintEditor.cpp b/Engine/Source/Editor/AnimationBlueprintEditor/Private/AnimationBlueprintEditor.cpp
--- a/Engine/Source/Editor/AnimationBlueprintEditor/Private/AnimationBlueprintEditor.cpp
+++ b/Engine/Source/Editor/AnimationBlueprintEditor/Private/AnimationBlueprintEditor.cpp
@@ -160,6 +160,7 @@ namespace AnimationBlueprintEditorTabs
 	const FName SlotNamesTab(TEXT("SkeletonSlotNames"));
 	const FName CurveNamesTab(TEXT("AnimCurveViewerTab"));
 	const FName PoseWatchTab(TEXT("PoseWatchManager"));
+	const FName FindReplaceTab(TEXT("FindReplaceTab"));
 };
 
 /////////////////////////////////////////////////////
diff --git a/Engine/Source/Editor/AnimationBlueprintEditor/Private/AnimationBlueprintEditor.h b/Engine/Source/Editor/AnimationBlueprintEditor/Private/AnimationBlueprintEditor.h
--- a/Engine/Source/Editor/AnimationBlueprintEditor/Private/AnimationBlueprintEditor.h
+++ b/Engine/Source/Editor/AnimationBlueprintEditor/Private/AnimationBlueprintEditor.h
@@ -83,6 +83,7 @@ namespace AnimationBlueprintEditorTabs
 	extern const FName SlotNamesTab;
 	extern const FName CurveNamesTab;
 	extern const FName PoseWatchTab;
+	extern const FName FindReplaceTab;
 };
 
 /**
diff --git a/Engine/Source/Editor/AnimationBlueprintEditor/Private/AnimationBlueprintEditorMode.cpp b/Engine/Source/Editor/AnimationBlueprintEditor/Private/AnimationBlueprintEditorMode.cpp
--- a/Engine/Source/Editor/AnimationBlueprintEditor/Private/AnimationBlueprintEditorMode.cpp
+++ b/Engine/Source/Editor/AnimationBlueprintEditor/Private/AnimationBlueprintEditorMode.cpp
@@ -25,6 +25,7 @@
 #include "Types/SlateEnums.h"
 #include "UObject/NameTypes.h"
 #include "UObject/WeakObjectPtr.h"
+#include "AnimAssetFindReplace.h"
 
 class UToolMenu;
 
@@ -79,7 +80,7 @@ FAnimationBlueprintEditorMode::FAnimationBlueprintEditorMode(const TSharedRef<FA
 		bIsTemplate = AnimBlueprint->bIsTemplate;
 	}
 
-	TabLayout = FTabManager::NewLayout( "Stanalone_AnimationBlueprintEditMode_Layout_v1.5" )
+	TabLayout = FTabManager::NewLayout( "Stanalone_AnimationBlueprintEditMode_Layout_v1.6" )
 		->AddArea
 		(
 			FTabManager::NewPrimaryArea()
@@ -130,6 +131,7 @@ FAnimationBlueprintEditorMode::FAnimationBlueprintEditorMode(const TSharedRef<FA
 						->SetSizeCoefficient(0.2f)
 						->AddTab(FBlueprintEditorTabs::CompilerResultsID, ETabState::ClosedTab)
 						->AddTab(FBlueprintEditorTabs::FindResultsID, ETabState::ClosedTab)
+						->AddTab(AnimationBlueprintEditorTabs::FindReplaceTab, ETabState::ClosedTab)
 					)
 				)
 				->Split
@@ -174,13 +176,14 @@ FAnimationBlueprintEditorMode::FAnimationBlueprintEditorMode(const TSharedRef<FA
 	TabFactories.RegisterFactory(PersonaModule.CreateAnimBlueprintPreviewTabFactory(InAnimationBlueprintEditor, InAnimationBlueprintEditor->GetPersonaToolkit()->GetPreviewScene()));
 	TabFactories.RegisterFactory(PersonaModule.CreateAnimBlueprintAssetOverridesTabFactory(InAnimationBlueprintEditor, InAnimationBlueprintEditor->GetPersonaToolkit()->GetAnimBlueprint(), InAnimationBlueprintEditor->OnPostUndo));
 	TabFactories.RegisterFactory(PersonaModule.CreatePoseWatchTabFactory(InAnimationBlueprintEditor));
+	TabFactories.RegisterFactory(PersonaModule.CreateCurveViewerTabFactory(InAnimationBlueprintEditor, InAnimationBlueprintEditor->GetSkeletonTree()->GetEditableSkeleton(), InAnimationBlueprintEditor->GetPersonaToolkit()->GetPreviewScene(), FOnObjectsSelected::CreateSP(&InAnimationBlueprintEditor.Get(), &FAnimationBlueprintEditor::HandleObjectsSelected)));
+	TabFactories.RegisterFactory(PersonaModule.CreateAnimAssetFindReplaceTabFactory(InAnimationBlueprintEditor, FAnimAssetFindReplaceConfig()));
 
 	if (!bIsTemplate)
 	{
 		ISkeletonEditorModule& SkeletonEditorModule = FModuleManager::LoadModuleChecked<ISkeletonEditorModule>("SkeletonEditor");
 		TabFactories.RegisterFactory(SkeletonEditorModule.CreateSkeletonTreeTabFactory(InAnimationBlueprintEditor, InAnimationBlueprintEditor->GetSkeletonTree()));
 		TabFactories.RegisterFactory(PersonaModule.CreateSkeletonSlotNamesTabFactory(InAnimationBlueprintEditor, InAnimationBlueprintEditor->GetSkeletonTree()->GetEditableSkeleton(), FOnObjectSelected::CreateSP(&InAnimationBlueprintEditor.Get(), &FAnimationBlueprintEditor::HandleObjectSelected)));
-		TabFactories.RegisterFactory(PersonaModule.CreateCurveViewerTabFactory(InAnimationBlueprintEditor, InAnimationBlueprintEditor->GetSkeletonTree()->GetEditableSkeleton(), InAnimationBlueprintEditor->GetPersonaToolkit()->GetPreviewScene(), FOnObjectsSelected::CreateSP(&InAnimationBlueprintEditor.Get(), &FAnimationBlueprintEditor::HandleObjectsSelected)));
 	}
 
 	// setup toolbar - clear existing toolbar extender from the BP mode
diff --git a/Engine/Source/Editor/AnimationBlueprintLibrary/Private/AnimPose.cpp b/Engine/Source/Editor/AnimationBlueprintLibrary/Private/AnimPose.cpp
--- a/Engine/Source/Editor/AnimationBlueprintLibrary/Private/AnimPose.cpp
+++ b/Engine/Source/Editor/AnimationBlueprintLibrary/Private/AnimPose.cpp
@@ -168,23 +168,11 @@ void FAnimPose::SetPose(const FAnimationPoseData& PoseData)
 		GenerateWorldSpaceTransforms();
 
 		const FBlendedCurve& Curve = PoseData.GetCurve();
-		for (TConstSetBitIterator It(Curve.ValidCurveWeights); It; ++It)
+		Curve.ForEachElement([&CurveNames = CurveNames, &CurveValues = CurveValues](const UE::Anim::FCurveElement& InElement)
 		{
-			const int32 EntryIndex = It.GetIndex();
-			const uint16 CurveNameUID = (*Curve.UIDToArrayIndexLUT).IsValidIndex(EntryIndex) ? (*Curve.UIDToArrayIndexLUT)[EntryIndex] : INDEX_NONE;
-
-			FSmartName CurveSmartName;
-			const USkeleton* Skeleton = ContextBoneContainer.GetSkeletonAsset();
-			if(Skeleton->GetSmartNameByUID(USkeleton::AnimCurveMappingName, CurveNameUID, CurveSmartName))
-			{
-				CurveNames.Add(CurveSmartName.DisplayName);
-				CurveValues.Add(Curve.CurveWeights[EntryIndex]);
-			}
-			else
-			{
-				ensureMsgf(false, TEXT("Unable to find SmartName for Curve with UID %i from Skeleton %s"), CurveNameUID, *Skeleton->GetPathName());
-			}
-		}
+			CurveNames.Add(InElement.Name);
+			CurveValues.Add(InElement.Value);
+		});
 	}
 	else
 	{
@@ -497,7 +485,7 @@ void UAnimPoseExtensions::GetReferencePose(USkeleton* Skeleton, FAnimPose& OutPo
 		}
 
 		FBoneContainer RequiredBones;
-		RequiredBones.InitializeTo(RequiredBoneIndexArray, FCurveEvaluationOption(false), *Skeleton);
+		RequiredBones.InitializeTo(RequiredBoneIndexArray, UE::Anim::FCurveFilterSettings(UE::Anim::ECurveFilterMode::DisallowAll), *Skeleton);
 
 		OutPose.Init(RequiredBones);
 		OutPose.SetToRefPose();
@@ -571,7 +559,7 @@ void UAnimPoseExtensions::GetAnimPoseAtTimeIntervals(const UAnimSequenceBase* An
 		}
 
 		FBoneContainer RequiredBones;
-		RequiredBones.InitializeTo(RequiredBoneIndexArray, FCurveEvaluationOption(EvaluationOptions.bEvaluateCurves), *AssetToUse);
+		RequiredBones.InitializeTo(RequiredBoneIndexArray, UE::Anim::FCurveFilterSettings(EvaluationOptions.bEvaluateCurves ? UE::Anim::ECurveFilterMode::None : UE::Anim::ECurveFilterMode::DisallowAll), *AssetToUse);
 		
 		RequiredBones.SetUseRAWData(EvaluationOptions.EvaluationType == EAnimDataEvalType::Raw);
 		RequiredBones.SetUseSourceData(EvaluationOptions.EvaluationType == EAnimDataEvalType::Source);
diff --git a/Engine/Source/Editor/AnimationBlueprintLibrary/Private/AnimationBlueprintLibrary.cpp b/Engine/Source/Editor/AnimationBlueprintLibrary/Private/AnimationBlueprintLibrary.cpp
--- a/Engine/Source/Editor/AnimationBlueprintLibrary/Private/AnimationBlueprintLibrary.cpp
+++ b/Engine/Source/Editor/AnimationBlueprintLibrary/Private/AnimationBlueprintLibrary.cpp
@@ -133,7 +133,7 @@ void UAnimationBlueprintLibrary::GetAnimationCurveNames(const UAnimSequence* Ani
 	{
 		auto GetCurveName = [](const auto& Curve) -> FName
 		{
-			return Curve.Name.DisplayName;
+			return Curve.GetName();
 		};
 
 		switch (CurveType)
@@ -1195,8 +1195,7 @@ void UAnimationBlueprintLibrary::AddCurve(UAnimSequence* AnimationSequence, FNam
 		if (bValidMetaData && bValidTransformCurveData )
 		{
 			// Add or retrieve the smartname
-			const FName Names[(int32)ESmartNameContainerType::SNCT_MAX] = { USkeleton::AnimCurveMappingName, USkeleton::AnimTrackCurveMappingName };
-			const bool bCurveAdded = AddCurveInternal(AnimationSequence, CurveName, Names[(int32)CurveContainer], CurveFlags, CurveType);
+			const bool bCurveAdded = AddCurveInternal(AnimationSequence, CurveName, CurveFlags, CurveType);
 
 			if (!bCurveAdded)
 			{
@@ -1227,10 +1226,10 @@ void UAnimationBlueprintLibrary::RemoveCurve(UAnimSequence* AnimationSequence, F
 {
 	if (AnimationSequence)
 	{
-		const FName ContainerName = RetrieveContainerNameForCurve(AnimationSequence, CurveName);
-		if (ContainerName != NAME_None)
+		const ERawCurveTrackTypes CurveType = RetrieveCurveTypeForCurve(AnimationSequence, CurveName);
+		if (CurveType != ERawCurveTrackTypes::RCT_MAX)
 		{
-			const bool bCurveRemoved = RemoveCurveInternal(AnimationSequence, CurveName, ContainerName, bRemoveNameFromSkeleton);
+			const bool bCurveRemoved = RemoveCurveInternal(AnimationSequence, CurveName, CurveType);
 		}
 		else
 		{
@@ -1375,27 +1374,27 @@ void UAnimationBlueprintLibrary::AddVectorCurveKeys(UAnimSequence* AnimationSequ
 	}
 }
 
-static void SetControllerCurveKeys(IAnimationDataController& Controller, FSmartName SmartName, const TArray<float>& Times, const TArray<float>& Values)
+static void SetControllerCurveKeys(IAnimationDataController& Controller, FName Name, const TArray<float>& Times, const TArray<float>& Values)
 {
 	const int32 NumKeys = Values.Num();
-	const FAnimationCurveIdentifier CurveId(SmartName, ERawCurveTrackTypes::RCT_Float);
+	const FAnimationCurveIdentifier CurveId(Name, ERawCurveTrackTypes::RCT_Float);
 	for (int32 KeyIndex = 0; KeyIndex < NumKeys; ++KeyIndex)
 	{
 		Controller.SetCurveKey(CurveId, { Times[KeyIndex], Values[KeyIndex] });
 	}
 }
 
-static void SetControllerCurveKeys(IAnimationDataController& Controller, FSmartName SmartName, const TArray<float>& Times, const TArray<FTransform>& Values)
+static void SetControllerCurveKeys(IAnimationDataController& Controller, FName Name, const TArray<float>& Times, const TArray<FTransform>& Values)
 {
 	const int32 NumKeys = Values.Num();
-	const FAnimationCurveIdentifier CurveId(SmartName, ERawCurveTrackTypes::RCT_Transform);
+	const FAnimationCurveIdentifier CurveId(Name, ERawCurveTrackTypes::RCT_Transform);
 	for (int32 KeyIndex = 0; KeyIndex < NumKeys; ++KeyIndex)
 	{
 		Controller.SetTransformCurveKey(CurveId, Times[KeyIndex], Values[KeyIndex]);
 	}
 }
 	
-static void SetControllerCurveKeys(IAnimationDataController& Controller, FSmartName CurveId, const TArray<float>& Times, const TArray<FVector>& Values)
+static void SetControllerCurveKeys(IAnimationDataController& Controller, FName Name, const TArray<float>& Times, const TArray<FVector>& Values)
 {
 	ensure(false);
 }
@@ -1405,31 +1404,20 @@ void UAnimationBlueprintLibrary::AddCurveKeysInternal(UAnimSequence* AnimationSe
 {
 	checkf(Times.Num() == KeyData.Num(), TEXT("Not enough key data supplied"));
 
-	const FName ContainerName = RetrieveContainerNameForCurve(AnimationSequence, CurveName);
-
-	if (ContainerName != NAME_None)
+	// Retrieve the curve by name
+	const FAnimationCurveIdentifier CurveId(CurveName, CurveType);
+	const CurveClass* Curve = static_cast<const CurveClass*>(AnimationSequence->GetDataModel()->FindCurve(CurveId));
+	if (Curve)
 	{
-		// Retrieve smart name for curve
-		const FSmartName CurveSmartName = RetrieveSmartNameForCurve(AnimationSequence, CurveName, ContainerName);
-
-		// Retrieve the curve by name
-		const FAnimationCurveIdentifier CurveId(CurveSmartName, CurveType);
-		const CurveClass* Curve = static_cast<const CurveClass*>(AnimationSequence->GetDataModel()->FindCurve(CurveId));
-		if (Curve)
-		{
-			IAnimationDataController& Controller = AnimationSequence->GetController();
-			SetControllerCurveKeys(Controller, CurveSmartName, Times, KeyData);
-		}
+		IAnimationDataController& Controller = AnimationSequence->GetController();
+		SetControllerCurveKeys(Controller, CurveName, Times, KeyData);
 	}
 }
 
-bool UAnimationBlueprintLibrary::AddCurveInternal(UAnimSequence* AnimationSequence, FName CurveName, FName ContainerName, int32 CurveFlags, ERawCurveTrackTypes SupportedCurveType)
+bool UAnimationBlueprintLibrary::AddCurveInternal(UAnimSequence* AnimationSequence, FName CurveName, int32 CurveFlags, ERawCurveTrackTypes SupportedCurveType)
 {
 	// Add or retrieve the smart name
-	FSmartName SmartCurveName;
-	AnimationSequence->GetSkeleton()->AddSmartNameAndModify(ContainerName, CurveName, SmartCurveName);	
-
-	FAnimationCurveIdentifier CurveId(SmartCurveName, SupportedCurveType);
+	FAnimationCurveIdentifier CurveId(CurveName, SupportedCurveType);
 
 	bool bCurveAdded = false;
 
@@ -1448,54 +1436,15 @@ bool UAnimationBlueprintLibrary::AddCurveInternal(UAnimSequence* AnimationSequen
 	return bCurveAdded;
 }
 
-bool UAnimationBlueprintLibrary::RemoveCurveInternal(UAnimSequence* AnimationSequence, FName CurveName, FName ContainerName, bool bRemoveNameFromSkeleton)
+bool UAnimationBlueprintLibrary::RemoveCurveInternal(UAnimSequence* AnimationSequence, FName CurveName, ERawCurveTrackTypes SupportedCurveType)
 {
 	checkf(AnimationSequence != nullptr, TEXT("Invalid Animation Sequence ptr"));
 	bool bRemoved = false;
-	SmartName::UID_Type UID = AnimationSequence->GetSkeleton()->GetUIDByName(ContainerName, CurveName);
-	if (UID != SmartName::MaxUID)
-	{
-		FSmartName SmartCurveName;
-		USkeleton* Skeleton = AnimationSequence->GetSkeleton();
-		checkf(Skeleton != nullptr, TEXT("Invalid Skeleton ptr"));
-		if (Skeleton->GetSmartNameByUID(ContainerName, UID, SmartCurveName))
-		{
-			IAnimationDataController& Controller = AnimationSequence->GetController();
-
-			const FAnimationCurveIdentifier CurveId(SmartCurveName, ContainerName == USkeleton::AnimTrackCurveMappingName ? ERawCurveTrackTypes::RCT_Transform : ERawCurveTrackTypes::RCT_Float);
-			bRemoved = Controller.RemoveCurve(CurveId);
-
-			if (bRemoveNameFromSkeleton)
-			{
-				// Ensure we are eligible to do this
-				bool bValidToRemove = true;
-
-				if (ContainerName == USkeleton::AnimTrackCurveMappingName)
-				{
-					// Make sure we do not remove bone names
-					bValidToRemove = DoesBoneCurveNameExistInternal(Skeleton, CurveName);
-				}
+	
+	IAnimationDataController& Controller = AnimationSequence->GetController();
 
-				if (bValidToRemove)
-				{
-					Skeleton->RemoveSmartnameAndModify(ContainerName, UID);
-				}
-				else
-				{
-					UE_LOG(LogAnimationBlueprintLibrary, Warning, TEXT("Cannot remove Curve Name %s from Skeleton %s"), *CurveName.ToString(), *AnimationSequence->GetSkeleton()->GetName());
-				}
-			}
-		}
-		else
-		{
-			UE_LOG(LogAnimationBlueprintLibrary, Warning, TEXT("Could not retrieve Smart Name for Curve Name %s from Skeleton %s"), *CurveName.ToString(), *AnimationSequence->GetSkeleton()->GetName());
-		}
-	}
-	else
-	{
-		// Name does not exist on skeleton
-		UE_LOG(LogAnimationBlueprintLibrary, Warning, TEXT("Could find for Curve Name %s in Skeleton %s"), *CurveName.ToString(), *AnimationSequence->GetSkeleton()->GetName());
-	}
+	const FAnimationCurveIdentifier CurveId(CurveName, SupportedCurveType);
+	bRemoved = Controller.RemoveCurve(CurveId);
 
 	return bRemoved;
 }
@@ -1527,12 +1476,6 @@ bool UAnimationBlueprintLibrary::DoesBoneNameExistInternal(USkeleton* Skeleton,
 	return Skeleton->GetReferenceSkeleton().FindBoneIndex(BoneName) != INDEX_NONE;
 }
 
-bool UAnimationBlueprintLibrary::DoesBoneCurveNameExistInternal(USkeleton* Skeleton, FName BoneName)
-{
-	checkf(Skeleton != nullptr, TEXT("Invalid Skeleton ptr"));
-	return Skeleton->GetUIDByName(USkeleton::AnimTrackCurveMappingName, BoneName) != SmartName::MaxUID;
-}
-
 void UAnimationBlueprintLibrary::GetFloatKeys(UAnimSequence* AnimationSequence, FName CurveName, TArray<float>& Times, TArray<float>& Values)
 {
 	if (AnimationSequence)
@@ -1569,33 +1512,19 @@ void UAnimationBlueprintLibrary::GetTransformationKeys(UAnimSequence* AnimationS
 	}
 }
 
-void UAnimationBlueprintLibrary::CopyAnimationCurveNamesToSkeleton(USkeleton* OldSkeleton, USkeleton* NewSkeleton, UAnimSequenceBase* SequenceBase, ERawCurveTrackTypes CurveType)
-{
-	IAnimationDataController& Controller = SequenceBase->GetController();
-	Controller.UpdateCurveNamesFromSkeleton(OldSkeleton, CurveType);
-	Controller.FindOrAddCurveNamesOnSkeleton(NewSkeleton, CurveType);	
-}
-
 template <typename DataType, typename CurveClass, ERawCurveTrackTypes CurveType>
 void UAnimationBlueprintLibrary::GetCurveKeysInternal(UAnimSequence* AnimationSequence, FName CurveName, TArray<float>& Times, TArray<DataType>& KeyData)
 {
 	checkf(AnimationSequence != nullptr, TEXT("Invalid Animation Sequence ptr"));
-	const FName ContainerName = RetrieveContainerNameForCurve(AnimationSequence, CurveName);
+	
+	// Retrieve the curve by name
+	const FAnimationCurveIdentifier CurveId(CurveName, CurveType);
+	const CurveClass* Curve = static_cast<const CurveClass*>(AnimationSequence->GetDataModel()->FindCurve(CurveId));
 
-	if (ContainerName != NAME_None)
+	if (Curve)
 	{
-		// Retrieve smart name for curve
-		const FSmartName CurveSmartName = RetrieveSmartNameForCurve(AnimationSequence, CurveName, ContainerName);
-
-		// Retrieve the curve by name
-		const FAnimationCurveIdentifier CurveId(CurveSmartName, CurveType);
-		const CurveClass* Curve = static_cast<const CurveClass*>(AnimationSequence->GetDataModel()->FindCurve(CurveId));
-
-		if (Curve)
-		{
-			Curve->GetKeys(Times, KeyData);
-			checkf(Times.Num() == KeyData.Num(), TEXT("Invalid key data retrieved from curve"));
-		}
+		Curve->GetKeys(Times, KeyData);
+		checkf(Times.Num() == KeyData.Num(), TEXT("Invalid key data retrieved from curve"));
 	}
 }
 
@@ -1605,20 +1534,9 @@ bool UAnimationBlueprintLibrary::DoesCurveExist(UAnimSequence* AnimationSequence
 
 	if (AnimationSequence)
 	{
-		FSmartName SmartName;
-		if (RetrieveSmartNameForCurve(AnimationSequence, CurveName, USkeleton::AnimTrackCurveMappingName, SmartName))
-		{
-			FAnimationCurveIdentifier CurveId(SmartName, CurveType);
-			const FAnimCurveBase* Curve = AnimationSequence->GetDataModel()->FindCurve(CurveId);
-			bExistingCurve = Curve != nullptr;
-		}
-
-		if (RetrieveSmartNameForCurve(AnimationSequence, CurveName, USkeleton::AnimCurveMappingName, SmartName))
-		{
-			FAnimationCurveIdentifier CurveId(SmartName, CurveType);
-			const FAnimCurveBase* Curve = AnimationSequence->GetDataModel()->FindCurve(CurveId);
-			bExistingCurve |= Curve != nullptr;
-		}
+		FAnimationCurveIdentifier CurveId(CurveName, CurveType);
+		const FAnimCurveBase* Curve = AnimationSequence->GetDataModel()->FindCurve(CurveId);
+		bExistingCurve = Curve != nullptr;
 	}
 	else
 	{
@@ -1628,42 +1546,19 @@ bool UAnimationBlueprintLibrary::DoesCurveExist(UAnimSequence* AnimationSequence
 	return bExistingCurve;
 }
 
-bool UAnimationBlueprintLibrary::DoesSmartNameExist(UAnimSequence* AnimationSequence, FName Name)
+ERawCurveTrackTypes UAnimationBlueprintLibrary::RetrieveCurveTypeForCurve(const UAnimSequence* AnimationSequence, FName CurveName)
 {
-	checkf(AnimationSequence != nullptr, TEXT("Invalid Animation Sequence ptr"));
-	FSmartName SmartName;
-	return AnimationSequence->GetSkeleton()->GetSmartNameByName(USkeleton::AnimTrackCurveMappingName, Name, SmartName) ||
-		AnimationSequence->GetSkeleton()->GetSmartNameByName(USkeleton::AnimCurveMappingName, Name, SmartName);
-}
-
-bool UAnimationBlueprintLibrary::RetrieveSmartNameForCurve(const UAnimSequence* AnimationSequence, FName CurveName, FName ContainerName, FSmartName& SmartName)
-{
-	checkf(AnimationSequence != nullptr, TEXT("Invalid Animation Sequence ptr"));
-	return AnimationSequence->GetSkeleton()->GetSmartNameByName(ContainerName, CurveName, SmartName);
-}
-
-FSmartName UAnimationBlueprintLibrary::RetrieveSmartNameForCurve(const UAnimSequence* AnimationSequence, FName CurveName, FName ContainerName)
-{
-	checkf(AnimationSequence != nullptr, TEXT("Invalid Animation Sequence ptr"));
-	FSmartName SmartCurveName;
-	AnimationSequence->GetSkeleton()->GetSmartNameByName(ContainerName, CurveName, SmartCurveName);
-	return SmartCurveName;
-}
-
-FName UAnimationBlueprintLibrary::RetrieveContainerNameForCurve(const UAnimSequence* AnimationSequence, FName CurveName)
-{
-	checkf(AnimationSequence != nullptr, TEXT("Invalid Animation Sequence ptr"));
-	const FName Names[(int32)ESmartNameContainerType::SNCT_MAX] = { USkeleton::AnimCurveMappingName, USkeleton::AnimTrackCurveMappingName };
-	for (int32 Index = 0; Index < (int32)ESmartNameContainerType::SNCT_MAX; ++Index)
+	if(AnimationSequence->GetDataModel()->FindCurve(FAnimationCurveIdentifier(CurveName, ERawCurveTrackTypes::RCT_Float)))
 	{
-		const FSmartNameMapping* CurveMapping = AnimationSequence->GetSkeleton()->GetSmartNameContainer(Names[Index]);
-		if (CurveMapping && CurveMapping->Exists(CurveName))
-		{
-			return Names[Index];
-		}
+		return ERawCurveTrackTypes::RCT_Float;
+	}
+	
+	if(AnimationSequence->GetDataModel()->FindCurve(FAnimationCurveIdentifier(CurveName, ERawCurveTrackTypes::RCT_Transform)))
+	{
+		return ERawCurveTrackTypes::RCT_Transform;
 	}
 
-	return NAME_None;
+	return ERawCurveTrackTypes::RCT_MAX;
 }
 
 void UAnimationBlueprintLibrary::AddMetaData(UAnimationAsset* AnimationAsset, TSubclassOf<UAnimMetaData> MetaDataClass, UAnimMetaData*& MetaDataInstance)
diff --git a/Engine/Source/Editor/AnimationBlueprintLibrary/Public/AnimationBlueprintLibrary.h b/Engine/Source/Editor/AnimationBlueprintLibrary/Public/AnimationBlueprintLibrary.h
--- a/Engine/Source/Editor/AnimationBlueprintLibrary/Public/AnimationBlueprintLibrary.h
+++ b/Engine/Source/Editor/AnimationBlueprintLibrary/Public/AnimationBlueprintLibrary.h
@@ -343,15 +343,17 @@ public:
 	static void AddCurveKeysInternal(UAnimSequence* AnimationSequence, FName CurveName, const TArray<float>& Times, const TArray<DataType>& KeyData);
 
 	// Returns true if successfully added, false if it was already existing
-	static bool AddCurveInternal(UAnimSequence* AnimationSequence, FName CurveName, FName ContainerName, int32 CurveFlags, ERawCurveTrackTypes SupportedCurveType);
-	static bool RemoveCurveInternal(UAnimSequence* AnimationSequence, FName CurveName, FName ContainerName, bool bRemoveNameFromSkeleton);
+	static bool AddCurveInternal(UAnimSequence* AnimationSequence, FName CurveName, int32 CurveFlags, ERawCurveTrackTypes SupportedCurveType);
+	static bool RemoveCurveInternal(UAnimSequence* AnimationSequence, FName CurveName, ERawCurveTrackTypes SupportedCurveType);
 
 	/** Checks whether or not the given Bone Name exist on the Skeleton referenced by the given Animation Sequence */
 	UFUNCTION(BlueprintCallable, Category = "AnimationBlueprintLibrary|Skeleton")
 	static void DoesBoneNameExist(UAnimSequence* AnimationSequence, FName BoneName, bool& bExists);
 
 	static bool DoesBoneNameExistInternal(USkeleton* Skeleton, FName BoneName);
-	static bool DoesBoneCurveNameExistInternal(USkeleton* Skeleton, FName BoneName);
+
+	UE_DEPRECATED(5.3, "This function is no longer used.")
+	static bool DoesBoneCurveNameExistInternal(USkeleton* Skeleton, FName BoneName) { return false; }
 
 	/** Retrieves, a multiple of, Float Key(s) from the specified Animation Curve inside of the given Animation Sequence */
 	UFUNCTION(BlueprintCallable, Category = "AnimationBlueprintLibrary|Curves")
@@ -369,8 +371,8 @@ public:
 	static void GetCurveKeysInternal(UAnimSequence* AnimationSequence, FName CurveName, TArray<float>& Times, TArray<DataType>& KeyData);
 
 	/** Ensures that any curve names that do not exist on the NewSkeleton are added to it, in which case the SmartName on the actual curve itself will also be updated */
-	UFUNCTION(BlueprintCallable, Category = "AnimationBlueprintLibrary|Curves")
-	static void CopyAnimationCurveNamesToSkeleton(USkeleton* OldSkeleton, USkeleton* NewSkeleton, UAnimSequenceBase* SequenceBase, ERawCurveTrackTypes CurveType);
+	UFUNCTION(BlueprintCallable, Category = "AnimationBlueprintLibrary|Curves", meta=(DeprecatedFunction, DeprecationMessage="It is no longer necessary to copy curve names to the skeleton. If metadata is required to be updated, please use the metadata setting APIs."))
+	static void CopyAnimationCurveNamesToSkeleton(USkeleton* OldSkeleton, USkeleton* NewSkeleton, UAnimSequenceBase* SequenceBase, ERawCurveTrackTypes CurveType) {}
 	
 	// Bone Tracks
 
@@ -398,11 +400,25 @@ public:
 	/** Checks whether or not the given Curve Name exist on the Skeleton referenced by the given Animation Sequence */
 	UFUNCTION(BlueprintCallable, Category = "AnimationBlueprintLibrary|Curves")
 	static bool DoesCurveExist(UAnimSequence* AnimationSequence, FName CurveName, ERawCurveTrackTypes CurveType);
-	static bool DoesSmartNameExist(UAnimSequence* AnimationSequence, FName Name);
+
+	UE_DEPRECATED(5.3, "This function is no longer used.")
+	static bool DoesSmartNameExist(UAnimSequence* AnimationSequence, FName Name) { return false; }
+
+	UE_DEPRECATED(5.3, "This function is no longer used.")
+	static FSmartName RetrieveSmartNameForCurve(const UAnimSequence* AnimationSequence, FName CurveName, FName ContainerName)
+	{
+		PRAGMA_DISABLE_DEPRECATION_WARNINGS
+		return FSmartName(CurveName, 0);
+		PRAGMA_ENABLE_DEPRECATION_WARNINGS
+	}
+
+	UE_DEPRECATED(5.3, "This function is no longer used.")
+	static bool RetrieveSmartNameForCurve(const UAnimSequence* AnimationSequence, FName CurveName, FName ContainerName, FSmartName& SmartName) { return false; }
+
+	UE_DEPRECATED(5.3, "This function is no longer used.")
+	static FName RetrieveContainerNameForCurve(const UAnimSequence* AnimationSequence, FName CurveName) { return NAME_None; }
 	
-	static FSmartName RetrieveSmartNameForCurve(const UAnimSequence* AnimationSequence, FName CurveName, FName ContainerName);
-	static bool RetrieveSmartNameForCurve(const UAnimSequence* AnimationSequence, FName CurveName, FName ContainerName, FSmartName& SmartName);
-	static FName RetrieveContainerNameForCurve(const UAnimSequence* AnimationSequence, FName CurveName);
+	static ERawCurveTrackTypes RetrieveCurveTypeForCurve(const UAnimSequence* AnimationSequence, FName CurveName);
 
 	// MetaData
 
diff --git a/Engine/Source/Editor/AnimationEditor/Private/AnimationEditor.cpp b/Engine/Source/Editor/AnimationEditor/Private/AnimationEditor.cpp
--- a/Engine/Source/Editor/AnimationEditor/Private/AnimationEditor.cpp
+++ b/Engine/Source/Editor/AnimationEditor/Private/AnimationEditor.cpp
@@ -103,6 +103,7 @@ const FName AnimationEditorTabs::AssetDetailsTab(TEXT("AnimAssetPropertiesTab"))
 const FName AnimationEditorTabs::CurveNamesTab(TEXT("AnimCurveViewerTab"));
 const FName AnimationEditorTabs::SlotNamesTab(TEXT("SkeletonSlotNames"));
 const FName AnimationEditorTabs::AnimMontageSectionsTab(TEXT("AnimMontageSections"));
+const FName AnimationEditorTabs::FindReplaceTab(TEXT("FindReplaceTab"));
 
 DEFINE_LOG_CATEGORY(LogAnimationEditor);
 
@@ -572,7 +573,7 @@ void FAnimationEditor::EditCurves(UAnimSequenceBase* InAnimSequence, const TArra
 
 	for(const FCurveEditInfo& CurveInfo : InCurveInfo)
 	{
-		CurveEditor.Pin()->AddCurve(CurveInfo.CurveDisplayName, CurveInfo.CurveColor, CurveInfo.Name, CurveInfo.Type, CurveInfo.CurveIndex, CurveInfo.OnCurveModified);
+		CurveEditor.Pin()->AddCurve(CurveInfo.CurveDisplayName, CurveInfo.CurveColor, CurveInfo.CurveName, CurveInfo.Type, CurveInfo.CurveIndex, CurveInfo.OnCurveModified);
 	}
 
 	CurveEditor.Pin()->ZoomToFit();
@@ -584,7 +585,7 @@ void FAnimationEditor::StopEditingCurves(const TArray<FCurveEditInfo>& InCurveIn
 	{
 		for(const FCurveEditInfo& CurveInfo : InCurveInfo)
 		{
-			CurveEditor.Pin()->RemoveCurve(CurveInfo.Name, CurveInfo.Type, CurveInfo.CurveIndex);
+			CurveEditor.Pin()->RemoveCurve(CurveInfo.CurveName, CurveInfo.Type, CurveInfo.CurveIndex);
 		}
 	}
 }
@@ -793,7 +794,7 @@ void FAnimationEditor::CopyCurveToSoundWave(const FAssetData& SoundWaveAssetData
 	// iterate over curves in anim data
 	for (const FFloatCurve& FloatCurve : Sequence->GetDataModel()->GetFloatCurves())
 	{
-		FRichCurve* Curve = FindOrAddCurve(CurveTable, FloatCurve.Name.DisplayName);
+		FRichCurve* Curve = FindOrAddCurve(CurveTable, FloatCurve.GetName());
 		*Curve = FloatCurve.FloatCurve; // copy data
 	}
 
diff --git a/Engine/Source/Editor/AnimationEditor/Private/AnimationEditor.h b/Engine/Source/Editor/AnimationEditor/Private/AnimationEditor.h
--- a/Engine/Source/Editor/AnimationEditor/Private/AnimationEditor.h
+++ b/Engine/Source/Editor/AnimationEditor/Private/AnimationEditor.h
@@ -57,6 +57,7 @@ namespace AnimationEditorTabs
 	extern const FName CurveNamesTab;
 	extern const FName SlotNamesTab;
 	extern const FName AnimMontageSectionsTab;
+	extern const FName FindReplaceTab;
 }
 
 class FAnimationEditor : public IAnimationEditor, public FGCObject, public FTickableEditorObject
diff --git a/Engine/Source/Editor/AnimationEditor/Private/AnimationEditorMode.cpp b/Engine/Source/Editor/AnimationEditor/Private/AnimationEditorMode.cpp
--- a/Engine/Source/Editor/AnimationEditor/Private/AnimationEditorMode.cpp
+++ b/Engine/Source/Editor/AnimationEditor/Private/AnimationEditorMode.cpp
@@ -15,6 +15,7 @@
 #include "PersonaModule.h"
 #include "Types/SlateEnums.h"
 #include "WorkflowOrientedApp/WorkflowCentricApplication.h"
+#include "AnimAssetFindReplace.h"
 
 FAnimationEditorMode::FAnimationEditorMode(TSharedRef<FWorkflowCentricApplication> InHostingApp, TSharedRef<ISkeletonTree> InSkeletonTree)
 	: FApplicationMode(AnimationEditorModes::AnimationEditorMode)
@@ -44,8 +45,9 @@ FAnimationEditorMode::FAnimationEditorMode(TSharedRef<FWorkflowCentricApplicatio
 	TabFactories.RegisterFactory(PersonaModule.CreateSkeletonSlotNamesTabFactory(InHostingApp, InSkeletonTree->GetEditableSkeleton(), FOnObjectSelected::CreateSP(&AnimationEditor.Get(), &FAnimationEditor::HandleObjectSelected)));
 	TabFactories.RegisterFactory(PersonaModule.CreateAnimNotifiesTabFactory(InHostingApp, InSkeletonTree->GetEditableSkeleton(), OnObjectsSelected));
 	TabFactories.RegisterFactory(PersonaModule.CreateAnimMontageSectionsTabFactory(InHostingApp, AnimationEditor->GetPersonaToolkit(), AnimationEditor->OnSectionsChanged));
+	TabFactories.RegisterFactory(PersonaModule.CreateAnimAssetFindReplaceTabFactory(InHostingApp, FAnimAssetFindReplaceConfig()));
 
-	TabLayout = FTabManager::NewLayout("Standalone_AnimationEditor_Layout_v1.4")
+	TabLayout = FTabManager::NewLayout("Standalone_AnimationEditor_Layout_v1.5")
 		->AddArea
 		(
 			FTabManager::NewPrimaryArea()
@@ -87,6 +89,7 @@ FAnimationEditorMode::FAnimationEditorMode(TSharedRef<FWorkflowCentricApplicatio
 						->SetHideTabWell(true)
 						->AddTab(AnimationEditorTabs::DocumentTab, ETabState::ClosedTab)
 						->AddTab(AnimationEditorTabs::CurveEditorTab, ETabState::ClosedTab)
+						->AddTab(AnimationEditorTabs::FindReplaceTab, ETabState::ClosedTab)
 						->SetForegroundTab(AnimationEditorTabs::DocumentTab)
 					)
 				)
diff --git a/Engine/Source/Editor/AnimationEditor/Public/IAnimationEditor.h b/Engine/Source/Editor/AnimationEditor/Public/IAnimationEditor.h
--- a/Engine/Source/Editor/AnimationEditor/Public/IAnimationEditor.h
+++ b/Engine/Source/Editor/AnimationEditor/Public/IAnimationEditor.h
@@ -26,29 +26,55 @@ public:
 	/** Support structure for EditCurves */
 	struct FCurveEditInfo
 	{
+		UE_DEPRECATED(5.3, "Please use the constructor that takes a FName.")
 		FCurveEditInfo(const FText& InCurveDisplayName, const FLinearColor& InCurveColor, const FSmartName& InName, ERawCurveTrackTypes InType, int32 InCurveIndex, FSimpleDelegate OnCurveModified = FSimpleDelegate())
 			: CurveDisplayName(InCurveDisplayName)
 			, CurveColor(InCurveColor)
-			, Name(InName)
+			, CurveName(InName.DisplayName)
 			, Type(InType)
 			, CurveIndex(InCurveIndex)
 			, OnCurveModified(OnCurveModified)
 		{}
 
+		UE_DEPRECATED(5.3, "Please use the constructor that takes a FName.")
 		FCurveEditInfo(const FSmartName& InName, ERawCurveTrackTypes InType, int32 InCurveIndex)
-			: Name(InName)
+			: CurveName(InName.DisplayName)
+			, Type(InType)
+			, CurveIndex(InCurveIndex)
+		{}
+		
+		FCurveEditInfo(const FText& InCurveDisplayName, const FLinearColor& InCurveColor, const FName& InName, ERawCurveTrackTypes InType, int32 InCurveIndex, FSimpleDelegate OnCurveModified = FSimpleDelegate())
+			: CurveDisplayName(InCurveDisplayName)
+			, CurveColor(InCurveColor)
+			, CurveName(InName)
+			, Type(InType)
+			, CurveIndex(InCurveIndex)
+			, OnCurveModified(OnCurveModified)
+		{}
+
+		FCurveEditInfo(const FName& InName, ERawCurveTrackTypes InType, int32 InCurveIndex)
+			: CurveName(InName)
 			, Type(InType)
 			, CurveIndex(InCurveIndex)
 		{}
 
 		bool operator==(const FCurveEditInfo& InCurveEditInfo) const
 		{
-			return Name.UID == InCurveEditInfo.Name.UID && Type == InCurveEditInfo.Type && CurveIndex == InCurveEditInfo.CurveIndex;
+			return CurveName == InCurveEditInfo.CurveName && Type == InCurveEditInfo.Type && CurveIndex == InCurveEditInfo.CurveIndex;
 		}
 
+		// removing deprecation for default copy operator/constructor to avoid deprecation warnings
+		PRAGMA_DISABLE_DEPRECATION_WARNINGS
+		FCurveEditInfo(const FCurveEditInfo&) = default;
+		FCurveEditInfo& operator=(const FCurveEditInfo&) = default;
+		PRAGMA_ENABLE_DEPRECATION_WARNINGS
+		
 		FText CurveDisplayName;
 		FLinearColor CurveColor;
+
+		UE_DEPRECATED(5.3, "Please use CurveName")
 		FSmartName Name;
+		FName CurveName;
 		ERawCurveTrackTypes Type;
 		int32 CurveIndex;
 		FSimpleDelegate OnCurveModified;
