Broke circular dependency:
AnimationDataController -> UnrealEd -> Engine -> AnimationDataController

diff --git a/Engine/Source/Developer/AnimationDataController/AnimationDataController.build.cs b/Engine/Source/Developer/AnimationDataController/AnimationDataController.build.cs
--- a/Engine/Source/Developer/AnimationDataController/AnimationDataController.build.cs
+++ b/Engine/Source/Developer/AnimationDataController/AnimationDataController.build.cs
@@ -10,14 +10,14 @@ namespace UnrealBuildTool.Rules
 				new string[]
 				{
 					"Core",
-					"CoreUObject",
-					"Engine"
+					"CoreUObject"
                 }
              );
 
 			PrivateDependencyModuleNames.AddRange(
 				new string[]
-				{					
+				{
+					"Engine",
 					"MovieScene"
 				}
 			);
