Duplicate fix from CL 26459834 over to AnimDataController as well

diff --git a/Engine/Source/Developer/AnimationDataController/Private/AnimDataController.cpp b/Engine/Source/Developer/AnimationDataController/Private/AnimDataController.cpp
--- a/Engine/Source/Developer/AnimationDataController/Private/AnimDataController.cpp
+++ b/Engine/Source/Developer/AnimationDataController/Private/AnimDataController.cpp
@@ -1978,9 +1978,10 @@ bool UAnimDataController::DuplicateAttribute(const FAnimationAttributeIdentifier
 				{
 					FTransaction Transaction = ConditionalTransaction(LOCTEXT("DuplicateAttribute", "Duplicating Animation Attribute"), bShouldTransact);
 
-					FAnimatedBoneAttribute& DuplicateAttribute = Model->AnimatedBoneAttributes.AddDefaulted_GetRef();
+					FAnimatedBoneAttribute DuplicateAttribute;
 					DuplicateAttribute.Identifier = NewAttributeIdentifier;
 					DuplicateAttribute.Curve = AttributePtr->Curve;
+					Model->AnimatedBoneAttributes.Add(DuplicateAttribute);
 
 					FAttributeAddedPayload Payload;
 					Payload.Identifier = NewAttributeIdentifier;
