Updated headers using UnrealCodeFixup to make sure dllstorage is on methods/staticvars instead of types. Part 4/n

diff --git a/Engine/Source/Developer/AnimationDataController/Public/AnimDataController.h b/Engine/Source/Developer/AnimationDataController/Public/AnimDataController.h
--- a/Engine/Source/Developer/AnimationDataController/Public/AnimDataController.h
+++ b/Engine/Source/Developer/AnimationDataController/Public/AnimDataController.h
@@ -15,8 +15,8 @@
 
 #include "AnimDataController.generated.h"
 
-UCLASS(BlueprintType)
-class ANIMATIONDATACONTROLLER_API UAnimDataController : public UObject, public IAnimationDataController
+UCLASS(MinimalAPI, BlueprintType)
+class UAnimDataController : public UObject, public IAnimationDataController
 {
 	GENERATED_BODY()
 
@@ -29,51 +29,51 @@ public:
 
 #if WITH_EDITOR
 	/** Begin IAnimationDataController overrides */
-	virtual void SetModel(TScriptInterface<IAnimationDataModel> InModel) override;
+	ANIMATIONDATACONTROLLER_API virtual void SetModel(TScriptInterface<IAnimationDataModel> InModel) override;
     virtual TScriptInterface<IAnimationDataModel> GetModelInterface() const override { return ModelInterface; }
 	virtual const IAnimationDataModel* const GetModel() const override { return Model.Get(); }
-	virtual void OpenBracket(const FText& InTitle, bool bShouldTransact = true) override;
-	virtual void CloseBracket(bool bShouldTransact = true) override;
-	virtual void SetNumberOfFrames(FFrameNumber Length, bool bShouldTransact = true) override;
-	virtual void ResizeNumberOfFrames(FFrameNumber NewLength, FFrameNumber T0, FFrameNumber T1, bool bShouldTransact = true) override;
-	virtual void ResizeInFrames(FFrameNumber Length, FFrameNumber T0, FFrameNumber T1, bool bShouldTransact = true) override;
-	virtual void SetFrameRate(FFrameRate FrameRate, bool bShouldTransact = true) override;
-	virtual void SetPlayLength(float Length, bool bShouldTransact = true) override;
-	virtual void ResizePlayLength(float NewLength, float T0, float T1, bool bShouldTransact = true) override;
-	virtual void Resize(float Length, float T0, float T1, bool bShouldTransact = true) override;
-	virtual bool AddBoneCurve(FName BoneName, bool bShouldTransact = true) override;
-	virtual int32 InsertBoneTrack(FName BoneName, int32 DesiredIndex, bool bShouldTransact = true) override;
-	virtual bool RemoveBoneTrack(FName BoneName, bool bShouldTransact = true) override;
-	virtual void RemoveAllBoneTracks(bool bShouldTransact = true) override;
-	virtual bool SetBoneTrackKeys(FName BoneName, const TArray<FVector3f>& PositionalKeys, const TArray<FQuat4f>& RotationalKeys, const TArray<FVector3f>& ScalingKeys, bool bShouldTransact = true) override;	
-	virtual bool SetBoneTrackKeys(FName BoneName, const TArray<FVector>& PositionalKeys, const TArray<FQuat>& RotationalKeys, const TArray<FVector>& ScalingKeys, bool bShouldTransact = true) override;	
-	virtual bool UpdateBoneTrackKeys(FName BoneName, const FInt32Range& KeyRangeToSet, const TArray<FVector3f>& PositionalKeys, const TArray<FQuat4f>& RotationalKeys, const TArray<FVector3f>& ScalingKeys, bool bShouldTransact = true) override;
-	virtual bool UpdateBoneTrackKeys(FName BoneName, const FInt32Range& KeyRangeToSet, const TArray<FVector>& PositionalKeys, const TArray<FQuat>& RotationalKeys, const TArray<FVector>& ScalingKeys, bool bShouldTransact = true) override;
-	virtual bool AddCurve(const FAnimationCurveIdentifier& CurveId, int32 CurveFlags = 0x00000004, bool bShouldTransact = true) override;
-	virtual bool DuplicateCurve(const FAnimationCurveIdentifier& CopyCurveId, const FAnimationCurveIdentifier& NewCurveId, bool bShouldTransact = true) override;
-	virtual bool RemoveCurve(const FAnimationCurveIdentifier& CurveId, bool bShouldTransact = true) override;
-	virtual void RemoveAllCurvesOfType(ERawCurveTrackTypes SupportedCurveType, bool bShouldTransact = true) override;
-	virtual bool SetCurveFlag(const FAnimationCurveIdentifier& CurveId, EAnimAssetCurveFlags Flag, bool bState = true, bool bShouldTransact = true) override;
-	virtual bool SetCurveFlags(const FAnimationCurveIdentifier& CurveId, int32 Flags, bool bShouldTransact = true) override;
-	virtual bool SetTransformCurveKeys(const FAnimationCurveIdentifier& CurveId, const TArray<FTransform>& TransformValues, const TArray<float>& TimeKeys, bool bShouldTransact = true) override;	
-	virtual bool SetTransformCurveKey(const FAnimationCurveIdentifier& CurveId, float Time, const FTransform& Value, bool bShouldTransact = true) override;
-	virtual bool RemoveTransformCurveKey(const FAnimationCurveIdentifier& CurveId, float Time, bool bShouldTransact = true) override;
-	virtual bool RenameCurve(const FAnimationCurveIdentifier& CurveToRenameId, const FAnimationCurveIdentifier& NewCurveId, bool bShouldTransact = true) override;
-	virtual bool SetCurveColor(const FAnimationCurveIdentifier& CurveId, FLinearColor Color, bool bShouldTransact = true) override;
-	virtual bool SetCurveComment(const FAnimationCurveIdentifier& CurveId, const FString& Comment, bool bShouldTransact = true) override;
-	virtual bool ScaleCurve(const FAnimationCurveIdentifier& CurveId, float Origin, float Factor, bool bShouldTransact = true) override;
-	virtual bool SetCurveKey(const FAnimationCurveIdentifier& CurveId, const FRichCurveKey& Key, bool bShouldTransact = true) override;	
-	virtual bool RemoveCurveKey(const FAnimationCurveIdentifier& CurveId, float Time, bool bShouldTransact = true) override;
-	virtual bool SetCurveKeys(const FAnimationCurveIdentifier& CurveId, const TArray<FRichCurveKey>& CurveKeys, bool bShouldTransact = true) override;
-	virtual bool SetCurveAttributes(const FAnimationCurveIdentifier& CurveId, const FCurveAttributes& Attributes, bool bShouldTransact = true) override;
-	virtual bool RemoveBoneTracksMissingFromSkeleton(const USkeleton* Skeleton, bool bShouldTransact = true) override;
-	virtual void UpdateAttributesFromSkeleton(const USkeleton* Skeleton, bool bShouldTransact = true) override;
-	virtual void NotifyPopulated() override;
-	virtual void ResetModel(bool bShouldTransact = true) override;
-	virtual bool AddAttribute(const FAnimationAttributeIdentifier& AttributeIdentifier, bool bShouldTransact = true) override;	
-    virtual bool RemoveAttribute(const FAnimationAttributeIdentifier& AttributeIdentifier, bool bShouldTransact = true) override;
-    virtual int32 RemoveAllAttributesForBone(const FName& BoneName, bool bShouldTransact = true) override;
-    virtual int32 RemoveAllAttributes(bool bShouldTransact = true) override;
+	ANIMATIONDATACONTROLLER_API virtual void OpenBracket(const FText& InTitle, bool bShouldTransact = true) override;
+	ANIMATIONDATACONTROLLER_API virtual void CloseBracket(bool bShouldTransact = true) override;
+	ANIMATIONDATACONTROLLER_API virtual void SetNumberOfFrames(FFrameNumber Length, bool bShouldTransact = true) override;
+	ANIMATIONDATACONTROLLER_API virtual void ResizeNumberOfFrames(FFrameNumber NewLength, FFrameNumber T0, FFrameNumber T1, bool bShouldTransact = true) override;
+	ANIMATIONDATACONTROLLER_API virtual void ResizeInFrames(FFrameNumber Length, FFrameNumber T0, FFrameNumber T1, bool bShouldTransact = true) override;
+	ANIMATIONDATACONTROLLER_API virtual void SetFrameRate(FFrameRate FrameRate, bool bShouldTransact = true) override;
+	ANIMATIONDATACONTROLLER_API virtual void SetPlayLength(float Length, bool bShouldTransact = true) override;
+	ANIMATIONDATACONTROLLER_API virtual void ResizePlayLength(float NewLength, float T0, float T1, bool bShouldTransact = true) override;
+	ANIMATIONDATACONTROLLER_API virtual void Resize(float Length, float T0, float T1, bool bShouldTransact = true) override;
+	ANIMATIONDATACONTROLLER_API virtual bool AddBoneCurve(FName BoneName, bool bShouldTransact = true) override;
+	ANIMATIONDATACONTROLLER_API virtual int32 InsertBoneTrack(FName BoneName, int32 DesiredIndex, bool bShouldTransact = true) override;
+	ANIMATIONDATACONTROLLER_API virtual bool RemoveBoneTrack(FName BoneName, bool bShouldTransact = true) override;
+	ANIMATIONDATACONTROLLER_API virtual void RemoveAllBoneTracks(bool bShouldTransact = true) override;
+	ANIMATIONDATACONTROLLER_API virtual bool SetBoneTrackKeys(FName BoneName, const TArray<FVector3f>& PositionalKeys, const TArray<FQuat4f>& RotationalKeys, const TArray<FVector3f>& ScalingKeys, bool bShouldTransact = true) override;	
+	ANIMATIONDATACONTROLLER_API virtual bool SetBoneTrackKeys(FName BoneName, const TArray<FVector>& PositionalKeys, const TArray<FQuat>& RotationalKeys, const TArray<FVector>& ScalingKeys, bool bShouldTransact = true) override;	
+	ANIMATIONDATACONTROLLER_API virtual bool UpdateBoneTrackKeys(FName BoneName, const FInt32Range& KeyRangeToSet, const TArray<FVector3f>& PositionalKeys, const TArray<FQuat4f>& RotationalKeys, const TArray<FVector3f>& ScalingKeys, bool bShouldTransact = true) override;
+	ANIMATIONDATACONTROLLER_API virtual bool UpdateBoneTrackKeys(FName BoneName, const FInt32Range& KeyRangeToSet, const TArray<FVector>& PositionalKeys, const TArray<FQuat>& RotationalKeys, const TArray<FVector>& ScalingKeys, bool bShouldTransact = true) override;
+	ANIMATIONDATACONTROLLER_API virtual bool AddCurve(const FAnimationCurveIdentifier& CurveId, int32 CurveFlags = 0x00000004, bool bShouldTransact = true) override;
+	ANIMATIONDATACONTROLLER_API virtual bool DuplicateCurve(const FAnimationCurveIdentifier& CopyCurveId, const FAnimationCurveIdentifier& NewCurveId, bool bShouldTransact = true) override;
+	ANIMATIONDATACONTROLLER_API virtual bool RemoveCurve(const FAnimationCurveIdentifier& CurveId, bool bShouldTransact = true) override;
+	ANIMATIONDATACONTROLLER_API virtual void RemoveAllCurvesOfType(ERawCurveTrackTypes SupportedCurveType, bool bShouldTransact = true) override;
+	ANIMATIONDATACONTROLLER_API virtual bool SetCurveFlag(const FAnimationCurveIdentifier& CurveId, EAnimAssetCurveFlags Flag, bool bState = true, bool bShouldTransact = true) override;
+	ANIMATIONDATACONTROLLER_API virtual bool SetCurveFlags(const FAnimationCurveIdentifier& CurveId, int32 Flags, bool bShouldTransact = true) override;
+	ANIMATIONDATACONTROLLER_API virtual bool SetTransformCurveKeys(const FAnimationCurveIdentifier& CurveId, const TArray<FTransform>& TransformValues, const TArray<float>& TimeKeys, bool bShouldTransact = true) override;	
+	ANIMATIONDATACONTROLLER_API virtual bool SetTransformCurveKey(const FAnimationCurveIdentifier& CurveId, float Time, const FTransform& Value, bool bShouldTransact = true) override;
+	ANIMATIONDATACONTROLLER_API virtual bool RemoveTransformCurveKey(const FAnimationCurveIdentifier& CurveId, float Time, bool bShouldTransact = true) override;
+	ANIMATIONDATACONTROLLER_API virtual bool RenameCurve(const FAnimationCurveIdentifier& CurveToRenameId, const FAnimationCurveIdentifier& NewCurveId, bool bShouldTransact = true) override;
+	ANIMATIONDATACONTROLLER_API virtual bool SetCurveColor(const FAnimationCurveIdentifier& CurveId, FLinearColor Color, bool bShouldTransact = true) override;
+	ANIMATIONDATACONTROLLER_API virtual bool SetCurveComment(const FAnimationCurveIdentifier& CurveId, const FString& Comment, bool bShouldTransact = true) override;
+	ANIMATIONDATACONTROLLER_API virtual bool ScaleCurve(const FAnimationCurveIdentifier& CurveId, float Origin, float Factor, bool bShouldTransact = true) override;
+	ANIMATIONDATACONTROLLER_API virtual bool SetCurveKey(const FAnimationCurveIdentifier& CurveId, const FRichCurveKey& Key, bool bShouldTransact = true) override;	
+	ANIMATIONDATACONTROLLER_API virtual bool RemoveCurveKey(const FAnimationCurveIdentifier& CurveId, float Time, bool bShouldTransact = true) override;
+	ANIMATIONDATACONTROLLER_API virtual bool SetCurveKeys(const FAnimationCurveIdentifier& CurveId, const TArray<FRichCurveKey>& CurveKeys, bool bShouldTransact = true) override;
+	ANIMATIONDATACONTROLLER_API virtual bool SetCurveAttributes(const FAnimationCurveIdentifier& CurveId, const FCurveAttributes& Attributes, bool bShouldTransact = true) override;
+	ANIMATIONDATACONTROLLER_API virtual bool RemoveBoneTracksMissingFromSkeleton(const USkeleton* Skeleton, bool bShouldTransact = true) override;
+	ANIMATIONDATACONTROLLER_API virtual void UpdateAttributesFromSkeleton(const USkeleton* Skeleton, bool bShouldTransact = true) override;
+	ANIMATIONDATACONTROLLER_API virtual void NotifyPopulated() override;
+	ANIMATIONDATACONTROLLER_API virtual void ResetModel(bool bShouldTransact = true) override;
+	ANIMATIONDATACONTROLLER_API virtual bool AddAttribute(const FAnimationAttributeIdentifier& AttributeIdentifier, bool bShouldTransact = true) override;	
+    ANIMATIONDATACONTROLLER_API virtual bool RemoveAttribute(const FAnimationAttributeIdentifier& AttributeIdentifier, bool bShouldTransact = true) override;
+    ANIMATIONDATACONTROLLER_API virtual int32 RemoveAllAttributesForBone(const FName& BoneName, bool bShouldTransact = true) override;
+    ANIMATIONDATACONTROLLER_API virtual int32 RemoveAllAttributes(bool bShouldTransact = true) override;
 	virtual bool SetAttributeKey(const FAnimationAttributeIdentifier& AttributeIdentifier, float Time, const void* KeyValue, const UScriptStruct* TypeStruct, bool bShouldTransact = true) override
 	{
 		return SetAttributeKey_Internal(AttributeIdentifier, Time, KeyValue, TypeStruct, bShouldTransact);
@@ -82,24 +82,24 @@ public:
 	{
 		return SetAttributeKeys_Internal(AttributeIdentifier, Times, KeyValues, TypeStruct, bShouldTransact);
 	}
-    virtual bool RemoveAttributeKey(const FAnimationAttributeIdentifier& AttributeIdentifier, float Time, bool bShouldTransact = true) override;	
-	virtual bool DuplicateAttribute(const FAnimationAttributeIdentifier& AttributeIdentifier, const FAnimationAttributeIdentifier& NewAttributeIdentifier, bool bShouldTransact = true) override;
-	virtual void UpdateWithSkeleton(USkeleton* TargetSkeleton, bool bShouldTransact = true) override;
-	virtual void PopulateWithExistingModel(TScriptInterface<IAnimationDataModel> InModel) override;
+    ANIMATIONDATACONTROLLER_API virtual bool RemoveAttributeKey(const FAnimationAttributeIdentifier& AttributeIdentifier, float Time, bool bShouldTransact = true) override;	
+	ANIMATIONDATACONTROLLER_API virtual bool DuplicateAttribute(const FAnimationAttributeIdentifier& AttributeIdentifier, const FAnimationAttributeIdentifier& NewAttributeIdentifier, bool bShouldTransact = true) override;
+	ANIMATIONDATACONTROLLER_API virtual void UpdateWithSkeleton(USkeleton* TargetSkeleton, bool bShouldTransact = true) override;
+	ANIMATIONDATACONTROLLER_API virtual void PopulateWithExistingModel(TScriptInterface<IAnimationDataModel> InModel) override;
 	virtual void InitializeModel() override {}
 protected:
-	virtual void NotifyBracketOpen() override;
-	virtual void NotifyBracketClosed() override;
+	ANIMATIONDATACONTROLLER_API virtual void NotifyBracketOpen() override;
+	ANIMATIONDATACONTROLLER_API virtual void NotifyBracketClosed() override;
 	/** End IAnimationDataController overrides */
 	
 private:
 	/** Internal functionality for setting Attribute curve key(s) */
-	bool SetAttributeKey_Internal(const FAnimationAttributeIdentifier& AttributeIdentifier, float Time, const void* KeyValue, const UScriptStruct* TypeStruct, bool bShouldTransact = true);
-	bool SetAttributeKeys_Internal(const FAnimationAttributeIdentifier& AttributeIdentifier, TArrayView<const float> Times, TArrayView<const void*> KeyValues, const UScriptStruct* TypeStruct, bool bShouldTransact = true);
+	ANIMATIONDATACONTROLLER_API bool SetAttributeKey_Internal(const FAnimationAttributeIdentifier& AttributeIdentifier, float Time, const void* KeyValue, const UScriptStruct* TypeStruct, bool bShouldTransact = true);
+	ANIMATIONDATACONTROLLER_API bool SetAttributeKeys_Internal(const FAnimationAttributeIdentifier& AttributeIdentifier, TArrayView<const float> Times, TArrayView<const void*> KeyValues, const UScriptStruct* TypeStruct, bool bShouldTransact = true);
 
 	/** Resizes the curve/attribute data stored on the model according to the provided new length and time at which to insert or remove time */
-	void ResizeCurves(float NewLength, bool bInserted, float T0, float T1, bool bShouldTransact = true);
-	void ResizeAttributes(float NewLength, bool bInserted, float T0, float T1, bool bShouldTransact = true);
+	ANIMATIONDATACONTROLLER_API void ResizeCurves(float NewLength, bool bInserted, float T0, float T1, bool bShouldTransact = true);
+	ANIMATIONDATACONTROLLER_API void ResizeAttributes(float NewLength, bool bInserted, float T0, float T1, bool bShouldTransact = true);
 #endif // WITH_EDITOR
 
 private: 
diff --git a/Engine/Source/Developer/AnimationDataController/Public/AnimDataControllerActions.h b/Engine/Source/Developer/AnimationDataController/Public/AnimDataControllerActions.h
--- a/Engine/Source/Developer/AnimationDataController/Public/AnimDataControllerActions.h
+++ b/Engine/Source/Developer/AnimationDataController/Public/AnimDataControllerActions.h
@@ -10,6 +10,8 @@
 #include "Animation/AnimTypes.h" 
 #include "Misc/FrameRate.h"
 
+#define UE_API ANIMATIONDATACONTROLLER_API
+
 class UObject;
 class IAnimationDataModel;
 class IAnimationDataController;
@@ -25,99 +27,99 @@ namespace Anim {
 * Each Action class represents an (invertable) operation mutating an IAnimationDataModel object utilizing a UAnimDataController. Allowing 
 * for a more granular approach to undo/redo-ing changes while also allowing for script-based interoperability.
 */
-class ANIMATIONDATACONTROLLER_API FAnimDataBaseAction : public FSwapChange
+class FAnimDataBaseAction : public FSwapChange
 {
 public:
-	virtual TUniquePtr<FChange> Execute(UObject* Object) final;
+	UE_API virtual TUniquePtr<FChange> Execute(UObject* Object) final;
 	virtual ~FAnimDataBaseAction() {}
-	virtual FString ToString() const override;
+	UE_API virtual FString ToString() const override;
 
 protected:
 	virtual TUniquePtr<FChange> ExecuteInternal(IAnimationDataModel* Model, IAnimationDataController* Controller) = 0;
 	virtual FString ToStringInternal() const = 0;
 };
 
-class ANIMATIONDATACONTROLLER_API FOpenBracketAction : public FAnimDataBaseAction
+class FOpenBracketAction : public FAnimDataBaseAction
 {
 public:
 	explicit FOpenBracketAction(const FString& InDescription) : Description(InDescription) {}
 	virtual ~FOpenBracketAction() {}
 protected:
 	FOpenBracketAction() {}
-	virtual TUniquePtr<FChange> ExecuteInternal(IAnimationDataModel* Model, IAnimationDataController* Controller) override;
-	virtual FString ToStringInternal() const override;
+	UE_API virtual TUniquePtr<FChange> ExecuteInternal(IAnimationDataModel* Model, IAnimationDataController* Controller) override;
+	UE_API virtual FString ToStringInternal() const override;
 
 protected:
 	FString Description;
 };
 
-class ANIMATIONDATACONTROLLER_API FCloseBracketAction : public FAnimDataBaseAction
+class FCloseBracketAction : public FAnimDataBaseAction
 {
 public:
 	explicit FCloseBracketAction(const FString& InDescription) : Description(InDescription) {}
 	virtual ~FCloseBracketAction() {}
 protected:
 	FCloseBracketAction() {}
-	virtual TUniquePtr<FChange> ExecuteInternal(IAnimationDataModel* Model, IAnimationDataController* Controller) override;
-	virtual FString ToStringInternal() const override;
+	UE_API virtual TUniquePtr<FChange> ExecuteInternal(IAnimationDataModel* Model, IAnimationDataController* Controller) override;
+	UE_API virtual FString ToStringInternal() const override;
 
 protected:
 	FString Description;
 };
 
-class ANIMATIONDATACONTROLLER_API FAddTrackAction : public FAnimDataBaseAction
+class FAddTrackAction : public FAnimDataBaseAction
 {
 public:
-	explicit FAddTrackAction(const FName& InName, TArray<FTransform>&& InTransformData);
+	UE_API explicit FAddTrackAction(const FName& InName, TArray<FTransform>&& InTransformData);
 	virtual ~FAddTrackAction() {}
 protected:
 	FAddTrackAction() {}
-	virtual TUniquePtr<FChange> ExecuteInternal(IAnimationDataModel* Model, IAnimationDataController* Controller) override;
-	virtual FString ToStringInternal() const override;
+	UE_API virtual TUniquePtr<FChange> ExecuteInternal(IAnimationDataModel* Model, IAnimationDataController* Controller) override;
+	UE_API virtual FString ToStringInternal() const override;
 
 protected:
 	FName Name;
 	TArray<FTransform> TransformData;
 };
 
-class ANIMATIONDATACONTROLLER_API FRemoveTrackAction : public FAnimDataBaseAction
+class FRemoveTrackAction : public FAnimDataBaseAction
 {
 public:
-	explicit FRemoveTrackAction(const FName& InName);
+	UE_API explicit FRemoveTrackAction(const FName& InName);
 	virtual ~FRemoveTrackAction() {}
 protected:
 	FRemoveTrackAction() {}
-	virtual TUniquePtr<FChange> ExecuteInternal(IAnimationDataModel* Model, IAnimationDataController* Controller) override;
-	virtual FString ToStringInternal() const override;
+	UE_API virtual TUniquePtr<FChange> ExecuteInternal(IAnimationDataModel* Model, IAnimationDataController* Controller) override;
+	UE_API virtual FString ToStringInternal() const override;
 
 protected:
 	FName Name;
 };
 
-class ANIMATIONDATACONTROLLER_API FSetTrackKeysAction : public FAnimDataBaseAction
+class FSetTrackKeysAction : public FAnimDataBaseAction
 {
 public:
-	explicit FSetTrackKeysAction(const FName& InName, TArray<FTransform>& InTransformData);
+	UE_API explicit FSetTrackKeysAction(const FName& InName, TArray<FTransform>& InTransformData);
 	virtual ~FSetTrackKeysAction() {}
 protected:
 	FSetTrackKeysAction() {}
-	virtual TUniquePtr<FChange> ExecuteInternal(IAnimationDataModel* Model, IAnimationDataController* Controller) override;
-	virtual FString ToStringInternal() const override;
+	UE_API virtual TUniquePtr<FChange> ExecuteInternal(IAnimationDataModel* Model, IAnimationDataController* Controller) override;
+	UE_API virtual FString ToStringInternal() const override;
 
 protected:
 	FName Name;
 	TArray<FTransform> TransformData;
 };
 
-class ANIMATIONDATACONTROLLER_API FResizePlayLengthInFramesAction : public FAnimDataBaseAction
+class FResizePlayLengthInFramesAction : public FAnimDataBaseAction
 {
 public:
-	explicit FResizePlayLengthInFramesAction(const IAnimationDataModel* InModel, FFrameNumber F0, FFrameNumber F1);
+	UE_API explicit FResizePlayLengthInFramesAction(const IAnimationDataModel* InModel, FFrameNumber F0, FFrameNumber F1);
 	virtual ~FResizePlayLengthInFramesAction() override {}
 protected:
 	FResizePlayLengthInFramesAction() {}
-	virtual TUniquePtr<FChange> ExecuteInternal(IAnimationDataModel* Model, IAnimationDataController* Controller) override;
-	virtual FString ToStringInternal() const override;
+	UE_API virtual TUniquePtr<FChange> ExecuteInternal(IAnimationDataModel* Model, IAnimationDataController* Controller) override;
+	UE_API virtual FString ToStringInternal() const override;
 
 protected:
 	FFrameNumber Length;
@@ -125,21 +127,21 @@ protected:
 	FFrameNumber Frame1;
 };
 
-class ANIMATIONDATACONTROLLER_API FSetFrameRateAction : public FAnimDataBaseAction
+class FSetFrameRateAction : public FAnimDataBaseAction
 {
 public:
-	explicit FSetFrameRateAction(const IAnimationDataModel* InModel);
+	UE_API explicit FSetFrameRateAction(const IAnimationDataModel* InModel);
 	virtual ~FSetFrameRateAction() override {}
 protected:
 	FSetFrameRateAction() {}
-	virtual TUniquePtr<FChange> ExecuteInternal(IAnimationDataModel* Model, IAnimationDataController* Controller) override;
-	virtual FString ToStringInternal() const override;
+	UE_API virtual TUniquePtr<FChange> ExecuteInternal(IAnimationDataModel* Model, IAnimationDataController* Controller) override;
+	UE_API virtual FString ToStringInternal() const override;
 
 protected:
 	FFrameRate FrameRate;
 };
 
-class ANIMATIONDATACONTROLLER_API FAddCurveAction : public FAnimDataBaseAction
+class FAddCurveAction : public FAnimDataBaseAction
 {
 public:
 	explicit FAddCurveAction(const FAnimationCurveIdentifier& InCurveId, int32 InFlags) : CurveId(InCurveId), Flags(InFlags) {}
@@ -147,23 +149,23 @@ public:
 	virtual ~FAddCurveAction() {}
 protected:
 	FAddCurveAction() {}
-	virtual TUniquePtr<FChange> ExecuteInternal(IAnimationDataModel* Model, IAnimationDataController* Controller) override;
-	virtual FString ToStringInternal() const override;
+	UE_API virtual TUniquePtr<FChange> ExecuteInternal(IAnimationDataModel* Model, IAnimationDataController* Controller) override;
+	UE_API virtual FString ToStringInternal() const override;
 
 protected:
 	FAnimationCurveIdentifier CurveId;
 	int32 Flags;
 };
 
-class ANIMATIONDATACONTROLLER_API FAddFloatCurveAction : public FAnimDataBaseAction
+class FAddFloatCurveAction : public FAnimDataBaseAction
 {
 public:
 	explicit FAddFloatCurveAction(const FAnimationCurveIdentifier& InCurveId, int32 InFlags, const TArray<FRichCurveKey>& InKeys, const FLinearColor& InColor) : CurveId(InCurveId), Flags(InFlags), Keys(InKeys), Color(InColor) {}
 	virtual ~FAddFloatCurveAction() {}
 protected:
 	FAddFloatCurveAction() {}
-	virtual TUniquePtr<FChange> ExecuteInternal(IAnimationDataModel* Model, IAnimationDataController* Controller) override;
-	virtual FString ToStringInternal() const override;
+	UE_API virtual TUniquePtr<FChange> ExecuteInternal(IAnimationDataModel* Model, IAnimationDataController* Controller) override;
+	UE_API virtual FString ToStringInternal() const override;
 
 protected:
 	FAnimationCurveIdentifier CurveId;
@@ -172,15 +174,15 @@ protected:
 	FLinearColor Color;
 };
 
-class ANIMATIONDATACONTROLLER_API FAddTransformCurveAction : public FAnimDataBaseAction
+class FAddTransformCurveAction : public FAnimDataBaseAction
 {
 public:
-	explicit FAddTransformCurveAction(const FAnimationCurveIdentifier& InCurveId, int32 InFlags, const FTransformCurve& InTransformCurve);
+	UE_API explicit FAddTransformCurveAction(const FAnimationCurveIdentifier& InCurveId, int32 InFlags, const FTransformCurve& InTransformCurve);
 	virtual ~FAddTransformCurveAction() {}
 protected:
 	FAddTransformCurveAction() {}
-	virtual TUniquePtr<FChange> ExecuteInternal(IAnimationDataModel* Model, IAnimationDataController* Controller) override;
-	virtual FString ToStringInternal() const override;
+	UE_API virtual TUniquePtr<FChange> ExecuteInternal(IAnimationDataModel* Model, IAnimationDataController* Controller) override;
+	UE_API virtual FString ToStringInternal() const override;
 
 protected:
 	FAnimationCurveIdentifier CurveId;
@@ -189,29 +191,29 @@ protected:
 	TArray<FRichCurveKey> SubCurveKeys[9];
 };
 
-class ANIMATIONDATACONTROLLER_API FRemoveCurveAction : public FAnimDataBaseAction
+class FRemoveCurveAction : public FAnimDataBaseAction
 {
 public:
 	explicit FRemoveCurveAction(const FAnimationCurveIdentifier& InCurveId) : CurveId(InCurveId) {}
 	virtual ~FRemoveCurveAction() {}
 protected:
 	FRemoveCurveAction() {}
-	virtual TUniquePtr<FChange> ExecuteInternal(IAnimationDataModel* Model, IAnimationDataController* Controller) override;
-	virtual FString ToStringInternal() const override;
+	UE_API virtual TUniquePtr<FChange> ExecuteInternal(IAnimationDataModel* Model, IAnimationDataController* Controller) override;
+	UE_API virtual FString ToStringInternal() const override;
 
 protected:
 	FAnimationCurveIdentifier CurveId;
 };
 
-class ANIMATIONDATACONTROLLER_API FSetCurveFlagsAction : public FAnimDataBaseAction
+class FSetCurveFlagsAction : public FAnimDataBaseAction
 {
 public:
 	explicit FSetCurveFlagsAction(const FAnimationCurveIdentifier& InCurveId, int32 InFlags, ERawCurveTrackTypes InCurveType) : CurveId(InCurveId), Flags(InFlags), CurveType(InCurveType) {}
 	virtual ~FSetCurveFlagsAction() {}
 protected:
 	FSetCurveFlagsAction() {}
-	virtual TUniquePtr<FChange> ExecuteInternal(IAnimationDataModel* Model, IAnimationDataController* Controller) override;
-	virtual FString ToStringInternal() const override;
+	UE_API virtual TUniquePtr<FChange> ExecuteInternal(IAnimationDataModel* Model, IAnimationDataController* Controller) override;
+	UE_API virtual FString ToStringInternal() const override;
 
 protected:
 	FAnimationCurveIdentifier CurveId;
@@ -219,30 +221,30 @@ protected:
 	ERawCurveTrackTypes CurveType;
 };
 
-class ANIMATIONDATACONTROLLER_API FRenameCurveAction : public FAnimDataBaseAction
+class FRenameCurveAction : public FAnimDataBaseAction
 {
 public:
 	explicit FRenameCurveAction(const FAnimationCurveIdentifier& InCurveId, const FAnimationCurveIdentifier& InNewCurveId) : CurveId(InCurveId), NewCurveId(InNewCurveId) {}
 	virtual ~FRenameCurveAction() {}
 protected:
 	FRenameCurveAction() {}
-	virtual TUniquePtr<FChange> ExecuteInternal(IAnimationDataModel* Model, IAnimationDataController* Controller) override;
-	virtual FString ToStringInternal() const override;
+	UE_API virtual TUniquePtr<FChange> ExecuteInternal(IAnimationDataModel* Model, IAnimationDataController* Controller) override;
+	UE_API virtual FString ToStringInternal() const override;
 
 protected:
 	FAnimationCurveIdentifier CurveId;
 	FAnimationCurveIdentifier NewCurveId;
 };
 
-class ANIMATIONDATACONTROLLER_API FScaleCurveAction : public FAnimDataBaseAction
+class FScaleCurveAction : public FAnimDataBaseAction
 {
 public:
-	explicit FScaleCurveAction(const FAnimationCurveIdentifier& InCurveId, float InOrigin, float InFactor, ERawCurveTrackTypes InCurveType);
+	UE_API explicit FScaleCurveAction(const FAnimationCurveIdentifier& InCurveId, float InOrigin, float InFactor, ERawCurveTrackTypes InCurveType);
 	virtual ~FScaleCurveAction() {}
 protected:
 	FScaleCurveAction() {}
-	virtual TUniquePtr<FChange> ExecuteInternal(IAnimationDataModel* Model, IAnimationDataController* Controller) override;
-	virtual FString ToStringInternal() const override;
+	UE_API virtual TUniquePtr<FChange> ExecuteInternal(IAnimationDataModel* Model, IAnimationDataController* Controller) override;
+	UE_API virtual FString ToStringInternal() const override;
 
 protected:
 	FAnimationCurveIdentifier CurveId;
@@ -251,15 +253,15 @@ protected:
 	float Factor;
 };
 
-class ANIMATIONDATACONTROLLER_API FAddRichCurveKeyAction : public FAnimDataBaseAction
+class FAddRichCurveKeyAction : public FAnimDataBaseAction
 {
 public:
 	explicit FAddRichCurveKeyAction(const FAnimationCurveIdentifier& InCurveId, const FRichCurveKey& InKey) : CurveId(InCurveId), Key(InKey) {}
 	virtual ~FAddRichCurveKeyAction() {}
 protected:
 	FAddRichCurveKeyAction() {}
-	virtual TUniquePtr<FChange> ExecuteInternal(IAnimationDataModel* Model, IAnimationDataController* Controller) override;
-	virtual FString ToStringInternal() const override;
+	UE_API virtual TUniquePtr<FChange> ExecuteInternal(IAnimationDataModel* Model, IAnimationDataController* Controller) override;
+	UE_API virtual FString ToStringInternal() const override;
 
 protected:
 	FAnimationCurveIdentifier CurveId;
@@ -267,179 +269,179 @@ protected:
 	FRichCurveKey Key;
 };
 
-class ANIMATIONDATACONTROLLER_API FSetRichCurveKeyAction : public FAnimDataBaseAction
+class FSetRichCurveKeyAction : public FAnimDataBaseAction
 {
 public:
 	explicit FSetRichCurveKeyAction(const FAnimationCurveIdentifier& InCurveId, const FRichCurveKey& InKey) : CurveId(InCurveId), Key(InKey) {}
 	virtual ~FSetRichCurveKeyAction() {}
 protected:
 	FSetRichCurveKeyAction() {}
-	virtual TUniquePtr<FChange> ExecuteInternal(IAnimationDataModel* Model, IAnimationDataController* Controller) override;
-	virtual FString ToStringInternal() const override;
+	UE_API virtual TUniquePtr<FChange> ExecuteInternal(IAnimationDataModel* Model, IAnimationDataController* Controller) override;
+	UE_API virtual FString ToStringInternal() const override;
 
 protected:
 	FAnimationCurveIdentifier CurveId;
 	FRichCurveKey Key;
 };
 
-class ANIMATIONDATACONTROLLER_API FRemoveRichCurveKeyAction : public FAnimDataBaseAction
+class FRemoveRichCurveKeyAction : public FAnimDataBaseAction
 {
 public:
 	explicit FRemoveRichCurveKeyAction(const FAnimationCurveIdentifier& InCurveId, const float InTime) : CurveId(InCurveId), Time(InTime) {}
 	virtual ~FRemoveRichCurveKeyAction() {}
 protected:
 	FRemoveRichCurveKeyAction() {}
-	virtual TUniquePtr<FChange> ExecuteInternal(IAnimationDataModel* Model, IAnimationDataController* Controller) override;
-	virtual FString ToStringInternal() const override;
+	UE_API virtual TUniquePtr<FChange> ExecuteInternal(IAnimationDataModel* Model, IAnimationDataController* Controller) override;
+	UE_API virtual FString ToStringInternal() const override;
 
 protected:
 	FAnimationCurveIdentifier CurveId;
 	float Time;
 };
 
-class ANIMATIONDATACONTROLLER_API FSetRichCurveKeysAction : public FAnimDataBaseAction
+class FSetRichCurveKeysAction : public FAnimDataBaseAction
 {
 public:
 	explicit FSetRichCurveKeysAction(const FAnimationCurveIdentifier& InCurveId, const TArray<FRichCurveKey>& InKeys) : CurveId(InCurveId), Keys(InKeys) {}
 	virtual ~FSetRichCurveKeysAction() {}
 protected:
 	FSetRichCurveKeysAction() {}
-	virtual TUniquePtr<FChange> ExecuteInternal(IAnimationDataModel* Model, IAnimationDataController* Controller) override;
-	virtual FString ToStringInternal() const override;
+	UE_API virtual TUniquePtr<FChange> ExecuteInternal(IAnimationDataModel* Model, IAnimationDataController* Controller) override;
+	UE_API virtual FString ToStringInternal() const override;
 
 protected:
 	FAnimationCurveIdentifier CurveId;
 	TArray<FRichCurveKey> Keys;
 };
 
-class ANIMATIONDATACONTROLLER_API FSetRichCurveAttributesAction : public FAnimDataBaseAction
+class FSetRichCurveAttributesAction : public FAnimDataBaseAction
 {
 public:
 	explicit FSetRichCurveAttributesAction(const FAnimationCurveIdentifier& InCurveId, const FCurveAttributes& InAttributes) : CurveId(InCurveId), Attributes(InAttributes) {}
 	virtual ~FSetRichCurveAttributesAction() {}
 protected:
 	FSetRichCurveAttributesAction() {}
-	virtual TUniquePtr<FChange> ExecuteInternal(IAnimationDataModel* Model, IAnimationDataController* Controller) override;
-	virtual FString ToStringInternal() const override;
+	UE_API virtual TUniquePtr<FChange> ExecuteInternal(IAnimationDataModel* Model, IAnimationDataController* Controller) override;
+	UE_API virtual FString ToStringInternal() const override;
 
 protected:
 	FAnimationCurveIdentifier CurveId;
 	FCurveAttributes Attributes;
 };
 
-class ANIMATIONDATACONTROLLER_API FSetCurveColorAction : public FAnimDataBaseAction
+class FSetCurveColorAction : public FAnimDataBaseAction
 {
 public:
 	explicit FSetCurveColorAction(const FAnimationCurveIdentifier& InCurveId, const FLinearColor& InColor) : CurveId(InCurveId), Color(InColor) {}
 	virtual ~FSetCurveColorAction() {}
 protected:
 	FSetCurveColorAction() {}
-	virtual TUniquePtr<FChange> ExecuteInternal(IAnimationDataModel* Model, IAnimationDataController* Controller) override;
-	virtual FString ToStringInternal() const override;
+	UE_API virtual TUniquePtr<FChange> ExecuteInternal(IAnimationDataModel* Model, IAnimationDataController* Controller) override;
+	UE_API virtual FString ToStringInternal() const override;
 
 protected:
 	FAnimationCurveIdentifier CurveId;
 	FLinearColor Color;
 };
 
-class ANIMATIONDATACONTROLLER_API FSetCurveCommentAction : public FAnimDataBaseAction
+class FSetCurveCommentAction : public FAnimDataBaseAction
 {
 public:
 	explicit FSetCurveCommentAction(const FAnimationCurveIdentifier& InCurveId, const FString& InComment) : CurveId(InCurveId), Comment(InComment) {}
 	virtual ~FSetCurveCommentAction() {}
 protected:
 	FSetCurveCommentAction() {}
-	virtual TUniquePtr<FChange> ExecuteInternal(IAnimationDataModel* Model, IAnimationDataController* Controller) override;
-	virtual FString ToStringInternal() const override;
+	UE_API virtual TUniquePtr<FChange> ExecuteInternal(IAnimationDataModel* Model, IAnimationDataController* Controller) override;
+	UE_API virtual FString ToStringInternal() const override;
 
 protected:
 	FAnimationCurveIdentifier CurveId;
 	FString Comment;
 };
 
-class ANIMATIONDATACONTROLLER_API FAddAtributeAction : public FAnimDataBaseAction
+class FAddAtributeAction : public FAnimDataBaseAction
 {
 public:
-	explicit FAddAtributeAction(const FAnimatedBoneAttribute& InAttribute);
+	UE_API explicit FAddAtributeAction(const FAnimatedBoneAttribute& InAttribute);
 	virtual ~FAddAtributeAction() {}
 protected:
 	FAddAtributeAction() {}
-	virtual TUniquePtr<FChange> ExecuteInternal(IAnimationDataModel* Model, IAnimationDataController* Controller) override;
-	virtual FString ToStringInternal() const override;
+	UE_API virtual TUniquePtr<FChange> ExecuteInternal(IAnimationDataModel* Model, IAnimationDataController* Controller) override;
+	UE_API virtual FString ToStringInternal() const override;
 
 protected:
 	FAnimationAttributeIdentifier AttributeId;
 	TArray<FAttributeKey> Keys;
 };
 
-class ANIMATIONDATACONTROLLER_API FRemoveAtributeAction : public FAnimDataBaseAction
+class FRemoveAtributeAction : public FAnimDataBaseAction
 {
 public:
 	explicit FRemoveAtributeAction(const FAnimationAttributeIdentifier& InAttributeId) : AttributeId(InAttributeId) {}
 	virtual ~FRemoveAtributeAction() {}
 protected:
 	FRemoveAtributeAction() {}
-	virtual TUniquePtr<FChange> ExecuteInternal(IAnimationDataModel* Model, IAnimationDataController* Controller) override;
-	virtual FString ToStringInternal() const override;
+	UE_API virtual TUniquePtr<FChange> ExecuteInternal(IAnimationDataModel* Model, IAnimationDataController* Controller) override;
+	UE_API virtual FString ToStringInternal() const override;
 
 protected:
 	FAnimationAttributeIdentifier AttributeId;
 };
 
-class ANIMATIONDATACONTROLLER_API FAddAtributeKeyAction : public FAnimDataBaseAction
+class FAddAtributeKeyAction : public FAnimDataBaseAction
 {
 public:
 	explicit FAddAtributeKeyAction(const FAnimationAttributeIdentifier& InAttributeId, const FAttributeKey& InKey) : AttributeId(InAttributeId), Key(InKey) {}
 	virtual ~FAddAtributeKeyAction() {}
 protected:
 	FAddAtributeKeyAction() {}
-	virtual TUniquePtr<FChange> ExecuteInternal(IAnimationDataModel* Model, IAnimationDataController* Controller) override;
-	virtual FString ToStringInternal() const override;
+	UE_API virtual TUniquePtr<FChange> ExecuteInternal(IAnimationDataModel* Model, IAnimationDataController* Controller) override;
+	UE_API virtual FString ToStringInternal() const override;
 
 protected:
 	FAnimationAttributeIdentifier AttributeId;
 	FAttributeKey Key;
 };
 
-class ANIMATIONDATACONTROLLER_API FSetAtributeKeyAction : public FAnimDataBaseAction
+class FSetAtributeKeyAction : public FAnimDataBaseAction
 {
 public:
 	explicit FSetAtributeKeyAction(const FAnimationAttributeIdentifier& InAttributeId, const FAttributeKey& InKey) : AttributeId(InAttributeId), Key(InKey) {}
 	virtual ~FSetAtributeKeyAction() {}
 protected:
 	FSetAtributeKeyAction() {}
-	virtual TUniquePtr<FChange> ExecuteInternal(IAnimationDataModel* Model, IAnimationDataController* Controller) override;
-	virtual FString ToStringInternal() const override;
+	UE_API virtual TUniquePtr<FChange> ExecuteInternal(IAnimationDataModel* Model, IAnimationDataController* Controller) override;
+	UE_API virtual FString ToStringInternal() const override;
 
 protected:
 	FAnimationAttributeIdentifier AttributeId;
 	FAttributeKey Key;
 };
 
-class ANIMATIONDATACONTROLLER_API FRemoveAtributeKeyAction : public FAnimDataBaseAction
+class FRemoveAtributeKeyAction : public FAnimDataBaseAction
 {
 public:
 	explicit FRemoveAtributeKeyAction(const FAnimationAttributeIdentifier& InAttributeId, float InTime) : AttributeId(InAttributeId), Time(InTime) {}
 	virtual ~FRemoveAtributeKeyAction() {}
 protected:
 	FRemoveAtributeKeyAction() {}
-	virtual TUniquePtr<FChange> ExecuteInternal(IAnimationDataModel* Model, IAnimationDataController* Controller) override;
-	virtual FString ToStringInternal() const override;
+	UE_API virtual TUniquePtr<FChange> ExecuteInternal(IAnimationDataModel* Model, IAnimationDataController* Controller) override;
+	UE_API virtual FString ToStringInternal() const override;
 
 protected:
 	FAnimationAttributeIdentifier AttributeId;
 	float Time;
 };
 
-class ANIMATIONDATACONTROLLER_API FSetAtributeKeysAction : public FAnimDataBaseAction
+class FSetAtributeKeysAction : public FAnimDataBaseAction
 {
 public:
-	explicit FSetAtributeKeysAction(const FAnimatedBoneAttribute& InAttribute);
+	UE_API explicit FSetAtributeKeysAction(const FAnimatedBoneAttribute& InAttribute);
 	virtual ~FSetAtributeKeysAction() {}
 protected:
 	FSetAtributeKeysAction() {}
-	virtual TUniquePtr<FChange> ExecuteInternal(IAnimationDataModel* Model, IAnimationDataController* Controller) override;
-	virtual FString ToStringInternal() const override;
+	UE_API virtual TUniquePtr<FChange> ExecuteInternal(IAnimationDataModel* Model, IAnimationDataController* Controller) override;
+	UE_API virtual FString ToStringInternal() const override;
 
 protected:
 	FAnimationAttributeIdentifier AttributeId;
@@ -450,4 +452,6 @@ protected:
 
 } // namespace UE
 
-#endif // WITH_EDITOR
\ No newline at end of file
+#endif // WITH_EDITOR
+
+#undef UE_API
diff --git a/Engine/Source/Developer/AssetTools/Public/AdvancedCopyCustomization.h b/Engine/Source/Developer/AssetTools/Public/AdvancedCopyCustomization.h
--- a/Engine/Source/Developer/AssetTools/Public/AdvancedCopyCustomization.h
+++ b/Engine/Source/Developer/AssetTools/Public/AdvancedCopyCustomization.h
@@ -8,15 +8,17 @@
 #include "IAssetTools.h"
 #include "AdvancedCopyCustomization.generated.h"
 
+#define UE_API ASSETTOOLS_API
 
 
-UCLASS()
-class ASSETTOOLS_API UAdvancedCopyCustomization : public UObject
+
+UCLASS(MinimalAPI)
+class UAdvancedCopyCustomization : public UObject
 {
 	GENERATED_BODY()
 
 public:
-	UAdvancedCopyCustomization(const class FObjectInitializer& ObjectInitializer);
+	UE_API UAdvancedCopyCustomization(const class FObjectInitializer& ObjectInitializer);
 
 	/* Getter for whether or not we should generate relative paths for this advanced copy */
 	bool GetShouldGenerateRelativePaths() const
@@ -43,7 +45,7 @@ public:
 	virtual bool CustomCopyValidate(const TMap<FString, FString>& OutPackagesAndDestinations) const { return true; };
 
 	/* Store the path of the package that caused this customization to be used */
-	void SetPackageThatInitiatedCopy(const FString& InBasePackage);
+	UE_API void SetPackageThatInitiatedCopy(const FString& InBasePackage);
 
 	const FString GetPackageThatInitiatedCopy() const { return PackageThatInitiatedCopy; }
 
@@ -58,3 +60,4 @@ protected:
 	FString PackageThatInitiatedCopy;
 };
 
+#undef UE_API
diff --git a/Engine/Source/Developer/AssetTools/Public/AssetToolsSettings.h b/Engine/Source/Developer/AssetTools/Public/AssetToolsSettings.h
--- a/Engine/Source/Developer/AssetTools/Public/AssetToolsSettings.h
+++ b/Engine/Source/Developer/AssetTools/Public/AssetToolsSettings.h
@@ -24,8 +24,8 @@ public:
 	FSoftClassPath AdvancedCopyCustomization;
 };
 
-UCLASS(config = Game, defaultconfig, meta = (DisplayName = "Asset Tools"))
-class ASSETTOOLS_API UAssetToolsSettings : public UDeveloperSettings
+UCLASS(MinimalAPI, config = Game, defaultconfig, meta = (DisplayName = "Asset Tools"))
+class UAssetToolsSettings : public UDeveloperSettings
 {
 	GENERATED_BODY()
 
diff --git a/Engine/Source/Developer/AssetTools/Public/AssetTypeActions/AssetTypeActions_AnimationAsset.h b/Engine/Source/Developer/AssetTools/Public/AssetTypeActions/AssetTypeActions_AnimationAsset.h
--- a/Engine/Source/Developer/AssetTools/Public/AssetTypeActions/AssetTypeActions_AnimationAsset.h
+++ b/Engine/Source/Developer/AssetTools/Public/AssetTypeActions/AssetTypeActions_AnimationAsset.h
@@ -9,34 +9,38 @@
 #include "EditorAnimUtils.h"
 #include "SSkeletonWidget.h"
 
+#define UE_API ASSETTOOLS_API
+
 class
 // UE_DEPRECATED(5.2, "The AssetDefinition system is replacing AssetTypeActions and UAssetDefinition_AnimationAsset replaced this.  Please see the Conversion Guide in AssetDefinition.h")
-ASSETTOOLS_API FAssetTypeActions_AnimationAsset : public FAssetTypeActions_Base
+FAssetTypeActions_AnimationAsset : public FAssetTypeActions_Base
 {
 public:
 	// IAssetTypeActions Implementation
 	virtual FText GetName() const override { return NSLOCTEXT("AssetTypeActions", "AssetTypeActions_AnimationAsset", "AnimationAsset"); }
 	virtual FColor GetTypeColor() const override { return FColor(80,123,72); }
 	virtual UClass* GetSupportedClass() const override { return UAnimationAsset::StaticClass(); }
-	virtual void OpenAssetEditor( const TArray<UObject*>& InObjects, TSharedPtr<class IToolkitHost> EditWithinLevelEditor = TSharedPtr<IToolkitHost>() ) override;
+	UE_API virtual void OpenAssetEditor( const TArray<UObject*>& InObjects, TSharedPtr<class IToolkitHost> EditWithinLevelEditor = TSharedPtr<IToolkitHost>() ) override;
 	virtual bool CanFilter() override { return false; }
 	virtual uint32 GetCategories() override { return EAssetTypeCategories::Animation; }
-	virtual class UThumbnailInfo* GetThumbnailInfo(UObject* Asset) const override;
+	UE_API virtual class UThumbnailInfo* GetThumbnailInfo(UObject* Asset) const override;
 
 private:
 
 	/** Handler for when FindSkeleton is selected */
-	void ExecuteFindSkeleton(TArray<TWeakObjectPtr<UAnimationAsset>> Objects);
+	UE_API void ExecuteFindSkeleton(TArray<TWeakObjectPtr<UAnimationAsset>> Objects);
 
 	/** Handler for when ReplaceSkeleton is selected */
-	void ExecuteReplaceSkeleton(TArray<TWeakObjectPtr<UAnimationAsset>> Objects);
+	UE_API void ExecuteReplaceSkeleton(TArray<TWeakObjectPtr<UAnimationAsset>> Objects);
 
 	/** Handle menu item for wanting to open asset in a new editor */
-	void ExecuteOpenInNewWindow(TArray<TWeakObjectPtr<UAnimationAsset>> Objects);
+	UE_API void ExecuteOpenInNewWindow(TArray<TWeakObjectPtr<UAnimationAsset>> Objects);
 
 	/** Open animation asset, will find existing editor if desired. */
-	void OpenAnimAssetEditor(const TArray<UObject*>& InObjects, bool bForceNewEditor, TSharedPtr<IToolkitHost> EditWithinLevelEditor);
+	UE_API void OpenAnimAssetEditor(const TArray<UObject*>& InObjects, bool bForceNewEditor, TSharedPtr<IToolkitHost> EditWithinLevelEditor);
 	
 	/** Replace skeleton when USkeleton is missing. Returns true only if Skeleton was replaced. */
-	bool ReplaceMissingSkeleton(TArray<TObjectPtr<UObject>> InAnimationAssets) const;
+	UE_API bool ReplaceMissingSkeleton(TArray<TObjectPtr<UObject>> InAnimationAssets) const;
 };
+
+#undef UE_API
diff --git a/Engine/Source/Developer/AssetTools/Public/AssetTypeActions/AssetTypeActions_Blueprint.h b/Engine/Source/Developer/AssetTools/Public/AssetTypeActions/AssetTypeActions_Blueprint.h
--- a/Engine/Source/Developer/AssetTools/Public/AssetTypeActions/AssetTypeActions_Blueprint.h
+++ b/Engine/Source/Developer/AssetTools/Public/AssetTypeActions/AssetTypeActions_Blueprint.h
@@ -7,13 +7,15 @@
 #include "Toolkits/IToolkitHost.h"
 #include "AssetTypeActions/AssetTypeActions_ClassTypeBase.h"
 
+#define UE_API ASSETTOOLS_API
+
 struct FAssetData;
 class IClassTypeActions;
 class UFactory;
 
 class
 // UE_DEPRECATED(5.2, "The AssetDefinition system is replacing AssetTypeActions and UAssetDefinition_Blueprint replaced this.  Please see the Conversion Guide in AssetDefinition.h")
-ASSETTOOLS_API FAssetTypeActions_Blueprint : public FAssetTypeActions_ClassTypeBase
+FAssetTypeActions_Blueprint : public FAssetTypeActions_ClassTypeBase
 {
 public:
 	// IAssetTypeActions Implementation
@@ -26,39 +28,41 @@ public:
 	// move.
 	virtual bool ShouldCallGetActions() const override { return GetSupportedClass() != UBlueprint::StaticClass(); }
 
-	virtual void GetActions(const TArray<UObject*>& InObjects, FToolMenuSection& Section) override;
-	virtual void OpenAssetEditor(const TArray<UObject*>& InObjects, TSharedPtr<class IToolkitHost> EditWithinLevelEditor = TSharedPtr<IToolkitHost>()) override;
-	virtual bool CanMerge() const override;
-	virtual void Merge(UObject* InObject) override;
-	virtual void Merge(UObject* BaseAsset, UObject* RemoteAsset, UObject* LocalAsset, const FOnMergeResolved& ResolutionCallback) override;
+	UE_API virtual void GetActions(const TArray<UObject*>& InObjects, FToolMenuSection& Section) override;
+	UE_API virtual void OpenAssetEditor(const TArray<UObject*>& InObjects, TSharedPtr<class IToolkitHost> EditWithinLevelEditor = TSharedPtr<IToolkitHost>()) override;
+	UE_API virtual bool CanMerge() const override;
+	UE_API virtual void Merge(UObject* InObject) override;
+	UE_API virtual void Merge(UObject* BaseAsset, UObject* RemoteAsset, UObject* LocalAsset, const FOnMergeResolved& ResolutionCallback) override;
 	virtual uint32 GetCategories() override { return EAssetTypeCategories::Blueprint | EAssetTypeCategories::Basic; }
-	virtual void PerformAssetDiff(UObject* Asset1, UObject* Asset2, const struct FRevisionInfo& OldRevision, const struct FRevisionInfo& NewRevision) const override;
-	virtual class UThumbnailInfo* GetThumbnailInfo(UObject* Asset) const override;
-	virtual FText GetAssetDescription(const FAssetData& AssetData) const override;
+	UE_API virtual void PerformAssetDiff(UObject* Asset1, UObject* Asset2, const struct FRevisionInfo& OldRevision, const struct FRevisionInfo& NewRevision) const override;
+	UE_API virtual class UThumbnailInfo* GetThumbnailInfo(UObject* Asset) const override;
+	UE_API virtual FText GetAssetDescription(const FAssetData& AssetData) const override;
 
 	// FAssetTypeActions_ClassTypeBase Implementation
-	virtual TWeakPtr<IClassTypeActions> GetClassTypeActions(const FAssetData& AssetData) const override;
+	UE_API virtual TWeakPtr<IClassTypeActions> GetClassTypeActions(const FAssetData& AssetData) const override;
 
 protected:
 	/** Whether or not this asset can create derived blueprints */
-	virtual bool CanCreateNewDerivedBlueprint() const;
+	UE_API virtual bool CanCreateNewDerivedBlueprint() const;
 
 	/** Return the factory responsible for creating this type of Blueprint */
-	virtual UFactory* GetFactoryForBlueprintType(UBlueprint* InBlueprint) const;
+	UE_API virtual UFactory* GetFactoryForBlueprintType(UBlueprint* InBlueprint) const;
 
 	/** Returns the tooltip to display when attempting to derive a Blueprint */
-	FText GetNewDerivedBlueprintTooltip(TWeakObjectPtr<UBlueprint> InObject);
+	UE_API FText GetNewDerivedBlueprintTooltip(TWeakObjectPtr<UBlueprint> InObject);
 
 	/** Returns TRUE if you can derive a Blueprint */
-	bool CanExecuteNewDerivedBlueprint(TWeakObjectPtr<UBlueprint> InObject);
+	UE_API bool CanExecuteNewDerivedBlueprint(TWeakObjectPtr<UBlueprint> InObject);
 
 private:
 	/** Handler for when EditDefaults is selected */
-	void ExecuteEditDefaults(TArray<TWeakObjectPtr<UBlueprint>> Objects);
+	UE_API void ExecuteEditDefaults(TArray<TWeakObjectPtr<UBlueprint>> Objects);
 
 	/** Handler for when NewDerivedBlueprint is selected */
-	void ExecuteNewDerivedBlueprint(TWeakObjectPtr<UBlueprint> InObject);
+	UE_API void ExecuteNewDerivedBlueprint(TWeakObjectPtr<UBlueprint> InObject);
 
 	/** Returns true if the blueprint is data only */
-	bool ShouldUseDataOnlyEditor( const UBlueprint* Blueprint ) const;
+	UE_API bool ShouldUseDataOnlyEditor( const UBlueprint* Blueprint ) const;
 };
+
+#undef UE_API
diff --git a/Engine/Source/Developer/AssetTools/Public/AssetTypeActions/AssetTypeActions_BlueprintGeneratedClass.h b/Engine/Source/Developer/AssetTools/Public/AssetTypeActions/AssetTypeActions_BlueprintGeneratedClass.h
--- a/Engine/Source/Developer/AssetTools/Public/AssetTypeActions/AssetTypeActions_BlueprintGeneratedClass.h
+++ b/Engine/Source/Developer/AssetTools/Public/AssetTypeActions/AssetTypeActions_BlueprintGeneratedClass.h
@@ -6,6 +6,8 @@
 #include "Toolkits/IToolkitHost.h"
 #include "AssetTypeActions/AssetTypeActions_ClassTypeBase.h"
 
+#define UE_API ASSETTOOLS_API
+
 class UBlueprintGeneratedClass;
 struct FAssetData;
 class IClassTypeActions;
@@ -13,13 +15,13 @@ class UFactory;
 
 class
 // UE_DEPRECATED(5.2, "The AssetDefinition system is replacing AssetTypeActions and UAssetDefinition_BlueprintGeneratedClass replaced this.  Please see the Conversion Guide in AssetDefinition.h")
-ASSETTOOLS_API FAssetTypeActions_BlueprintGeneratedClass : public FAssetTypeActions_ClassTypeBase
+FAssetTypeActions_BlueprintGeneratedClass : public FAssetTypeActions_ClassTypeBase
 {
 public:
 	// IAssetTypeActions Implementation
 	virtual FText GetName() const override { return NSLOCTEXT("AssetTypeActions", "AssetTypeActions_BlueprintGeneratedClass", "Compiled Blueprint Class"); }
 	virtual FColor GetTypeColor() const override { return FColor(133, 173, 255); }
-	virtual UClass* GetSupportedClass() const override;
+	UE_API virtual UClass* GetSupportedClass() const override;
 
 	// We don't want BPGCs to show up in CB/Ref viewer filters since the BP filters already include them
 	virtual uint32 GetCategories() override { return EAssetTypeCategories::None; }
@@ -28,6 +30,8 @@ public:
 	// End IAssetTypeActions Implementation
 
 	// FAssetTypeActions_ClassTypeBase Implementation
-	virtual TWeakPtr<IClassTypeActions> GetClassTypeActions(const FAssetData& AssetData) const override;
+	UE_API virtual TWeakPtr<IClassTypeActions> GetClassTypeActions(const FAssetData& AssetData) const override;
 	// End FAssetTypeActions_ClassTypeBase Implementation
 };
+
+#undef UE_API
diff --git a/Engine/Source/Developer/AssetTools/Public/AssetTypeActions/AssetTypeActions_ClassTypeBase.h b/Engine/Source/Developer/AssetTools/Public/AssetTypeActions/AssetTypeActions_ClassTypeBase.h
--- a/Engine/Source/Developer/AssetTools/Public/AssetTypeActions/AssetTypeActions_ClassTypeBase.h
+++ b/Engine/Source/Developer/AssetTools/Public/AssetTypeActions/AssetTypeActions_ClassTypeBase.h
@@ -12,7 +12,7 @@ class IClassTypeActions;
 /** Base class for "class type" assets (C++ classes and Blueprints */
 class
 // UE_DEPRECATED(5.2, "The AssetDefinition system is replacing AssetTypeActions and this was replaced by AssetDefinition_ClassTypeBase.  Please see the Conversion Guide in AssetDefinition.h")
-ASSETTOOLS_API FAssetTypeActions_ClassTypeBase : public FAssetTypeActions_Base
+FAssetTypeActions_ClassTypeBase : public FAssetTypeActions_Base
 {
 public:
 	virtual bool CanLocalize() const override { return false; }
diff --git a/Engine/Source/Developer/AssetTools/Public/AssetTypeActions/AssetTypeActions_DataAsset.h b/Engine/Source/Developer/AssetTools/Public/AssetTypeActions/AssetTypeActions_DataAsset.h
--- a/Engine/Source/Developer/AssetTools/Public/AssetTypeActions/AssetTypeActions_DataAsset.h
+++ b/Engine/Source/Developer/AssetTools/Public/AssetTypeActions/AssetTypeActions_DataAsset.h
@@ -7,9 +7,11 @@
 #include "Engine/DataAsset.h"
 #include "UObject/WeakObjectPtrTemplates.h"
 
+#define UE_API ASSETTOOLS_API
+
 class
 // UE_DEPRECATED(5.2, "The AssetDefinition system is replacing AssetTypeActions and UAssetDefinition_DataAsset replaced this.  Please see the Conversion Guide in AssetDefinition.h")
-ASSETTOOLS_API FAssetTypeActions_DataAsset : public FAssetTypeActions_Base
+FAssetTypeActions_DataAsset : public FAssetTypeActions_Base
 {
 public:
 	// IAssetTypeActions Implementation
@@ -17,10 +19,12 @@ public:
 	virtual FColor GetTypeColor() const override { return FColor(201, 29, 85); }
 	virtual UClass* GetSupportedClass() const override { return UDataAsset::StaticClass(); }
 	virtual uint32 GetCategories() override { return EAssetTypeCategories::Misc; }
-	virtual FText GetDisplayNameFromAssetData(const FAssetData& AssetData) const override;
-	virtual void GetActions(const TArray<UObject*>& InObjects, struct FToolMenuSection& Section) override;
-	virtual void PerformAssetDiff(UObject* OldAsset, UObject* NewAsset, const FRevisionInfo& OldRevision, const FRevisionInfo& NewRevision) const override;
+	UE_API virtual FText GetDisplayNameFromAssetData(const FAssetData& AssetData) const override;
+	UE_API virtual void GetActions(const TArray<UObject*>& InObjects, struct FToolMenuSection& Section) override;
+	UE_API virtual void PerformAssetDiff(UObject* OldAsset, UObject* NewAsset, const FRevisionInfo& OldRevision, const FRevisionInfo& NewRevision) const override;
 
 private:
-	void ExecuteChangeDataAssetClass(TArray<TWeakObjectPtr<UDataAsset>> InDataAssets);
+	UE_API void ExecuteChangeDataAssetClass(TArray<TWeakObjectPtr<UDataAsset>> InDataAssets);
 };
+
+#undef UE_API
diff --git a/Engine/Source/Developer/AssetTools/Public/AssetTypeActions_CSVAssetBase.h b/Engine/Source/Developer/AssetTools/Public/AssetTypeActions_CSVAssetBase.h
--- a/Engine/Source/Developer/AssetTools/Public/AssetTypeActions_CSVAssetBase.h
+++ b/Engine/Source/Developer/AssetTools/Public/AssetTypeActions_CSVAssetBase.h
@@ -5,9 +5,11 @@
 #include "CoreMinimal.h"
 #include "AssetTypeActions_Base.h"
 
+#define UE_API ASSETTOOLS_API
+
 class 
 // UE_DEPRECATED(5.2, "The AssetDefinition system is replacing AssetTypeActions and nothing replaced this, just subclass from UAssetDefinitionDefault.  If you needed the ExecuteFindSourceFileInExplorer, you can now find that in FindSourceFileInExplorer.h.  Please see the Conversion Guide in AssetDefinition.h")
-ASSETTOOLS_API FAssetTypeActions_CSVAssetBase : public FAssetTypeActions_Base
+FAssetTypeActions_CSVAssetBase : public FAssetTypeActions_Base
 {
 public:
 	virtual FColor GetTypeColor() const override { return FColor(62, 140, 35); }
@@ -16,11 +18,13 @@ public:
 
 protected:
 	/** Handler for opening the source file for this asset */
-	void ExecuteFindSourceFileInExplorer(TArray<FString> Filenames, TArray<FString> OverrideExtensions);
+	UE_API void ExecuteFindSourceFileInExplorer(TArray<FString> Filenames, TArray<FString> OverrideExtensions);
 
 	/** Determine whether the find source file in explorer editor command can execute or not */
-	bool CanExecuteFindSourceFileInExplorer(TArray<FString> Filenames, TArray<FString> OverrideExtensions) const;
+	UE_API bool CanExecuteFindSourceFileInExplorer(TArray<FString> Filenames, TArray<FString> OverrideExtensions) const;
 
 	/** Verify the specified filename exists */
-	bool VerifyFileExists(const FString& InFileName) const;
+	UE_API bool VerifyFileExists(const FString& InFileName) const;
 };
+
+#undef UE_API
diff --git a/Engine/Source/Developer/AssetTools/Public/CollectionAssetManagement.h b/Engine/Source/Developer/AssetTools/Public/CollectionAssetManagement.h
--- a/Engine/Source/Developer/AssetTools/Public/CollectionAssetManagement.h
+++ b/Engine/Source/Developer/AssetTools/Public/CollectionAssetManagement.h
@@ -7,62 +7,64 @@
 #include "CollectionManagerTypes.h"
 #include "Styling/SlateTypes.h"
 
+#define UE_API ASSETTOOLS_API
+
 class ICollectionContainer;
 
 /** Handles the collection management for the given assets */
-class ASSETTOOLS_API FCollectionAssetManagement
+class FCollectionAssetManagement
 {
 public:
 	/** Constructor */
 	UE_DEPRECATED(5.6, "Use the ICollectionContainer constructor instead.")
-	FCollectionAssetManagement();
-	explicit FCollectionAssetManagement(const TSharedRef<ICollectionContainer>& InCollectionContainer);
+	UE_API FCollectionAssetManagement();
+	UE_API explicit FCollectionAssetManagement(const TSharedRef<ICollectionContainer>& InCollectionContainer);
 	FCollectionAssetManagement(const FCollectionAssetManagement&) = delete;
 
 	/** Destructor */
-	~FCollectionAssetManagement();
+	UE_API ~FCollectionAssetManagement();
 
 	FCollectionAssetManagement& operator=(const FCollectionAssetManagement&) = delete;
 
 	/** Set the assets that we are currently observing and managing the collection state of */
-	void SetCurrentAssets(const TArray<FAssetData>& CurrentAssets);
+	UE_API void SetCurrentAssets(const TArray<FAssetData>& CurrentAssets);
 
 	/** Set the asset paths that we are currently observing and managing the collection state of */
-	void SetCurrentAssetPaths(const TArray<FSoftObjectPath>& CurrentAssets);
+	UE_API void SetCurrentAssetPaths(const TArray<FSoftObjectPath>& CurrentAssets);
 
 	/** Get the number of assets in the current set */
-	int32 GetCurrentAssetCount() const;
+	UE_API int32 GetCurrentAssetCount() const;
 
 	/** Add the current assets to the given collection */
-	void AddCurrentAssetsToCollection(FCollectionNameType InCollectionKey);
+	UE_API void AddCurrentAssetsToCollection(FCollectionNameType InCollectionKey);
 
 	/** Remove the current assets from the given collection */
-	void RemoveCurrentAssetsFromCollection(FCollectionNameType InCollectionKey);
+	UE_API void RemoveCurrentAssetsFromCollection(FCollectionNameType InCollectionKey);
 
 	/** Return whether or not the given collection should be enabled in any management UIs */
-	bool IsCollectionEnabled(FCollectionNameType InCollectionKey) const;
+	UE_API bool IsCollectionEnabled(FCollectionNameType InCollectionKey) const;
 
 	/** Get the check box state the given collection should use in any management UIs */
-	ECheckBoxState GetCollectionCheckState(FCollectionNameType InCollectionKey) const;
+	UE_API ECheckBoxState GetCollectionCheckState(FCollectionNameType InCollectionKey) const;
 
 private:
 	/** Update the internal state used to track the check box status for each collection */
-	void UpdateAssetManagementState();
+	UE_API void UpdateAssetManagementState();
 
 	/** Handles an on collection renamed event */
-	void HandleCollectionRenamed(ICollectionContainer&, const FCollectionNameType& OriginalCollection, const FCollectionNameType& NewCollection);
+	UE_API void HandleCollectionRenamed(ICollectionContainer&, const FCollectionNameType& OriginalCollection, const FCollectionNameType& NewCollection);
 
 	/** Handles an on collection updated event */
-	void HandleCollectionUpdated(ICollectionContainer&, const FCollectionNameType& Collection);
+	UE_API void HandleCollectionUpdated(ICollectionContainer&, const FCollectionNameType& Collection);
 
 	/** Handles an on collection destroyed event */
-	void HandleCollectionDestroyed(ICollectionContainer&, const FCollectionNameType& Collection);
+	UE_API void HandleCollectionDestroyed(ICollectionContainer&, const FCollectionNameType& Collection);
 
 	/** Handles assets being added to a collection */
-	void HandleAssetsAddedToCollection(ICollectionContainer&, const FCollectionNameType& Collection, TConstArrayView<FSoftObjectPath> AssetsAdded);
+	UE_API void HandleAssetsAddedToCollection(ICollectionContainer&, const FCollectionNameType& Collection, TConstArrayView<FSoftObjectPath> AssetsAdded);
 
 	/** Handles assets being removed from a collection */
-	void HandleAssetsRemovedFromCollection(ICollectionContainer&, const FCollectionNameType& Collection, TConstArrayView<FSoftObjectPath> AssetsRemoved);
+	UE_API void HandleAssetsRemovedFromCollection(ICollectionContainer&, const FCollectionNameType& Collection, TConstArrayView<FSoftObjectPath> AssetsRemoved);
 
 	TSharedRef<ICollectionContainer> CollectionContainer;
 
@@ -79,3 +81,5 @@ private:
 	FDelegateHandle OnAssetsAddedHandle;
 	FDelegateHandle OnAssetsRemovedHandle;
 };
+
+#undef UE_API
diff --git a/Engine/Source/Developer/AssetTools/Public/PackageMigrationContext.h b/Engine/Source/Developer/AssetTools/Public/PackageMigrationContext.h
--- a/Engine/Source/Developer/AssetTools/Public/PackageMigrationContext.h
+++ b/Engine/Source/Developer/AssetTools/Public/PackageMigrationContext.h
@@ -6,6 +6,8 @@
 #include "Containers/UnrealString.h"
 #include "UObject/GCObject.h"
 
+#define UE_API ASSETTOOLS_API
+
 class UPackage;
 class UAssetToolsImpl;
 
@@ -17,33 +19,33 @@ namespace Private
 	struct FPackageMigrationImpl;
 }
 
-struct ASSETTOOLS_API FPackageMigrationContext : public FGCObject
+struct FPackageMigrationContext : public FGCObject
 {
 	enum class EPackageMigrationStep;
 
-	EPackageMigrationStep GetCurrentStep() const;
+	UE_API EPackageMigrationStep GetCurrentStep() const;
 
 	// Notify the user that files where migrated. It should tell the user what was migrated and where it is in the destination.
-	void AddSucessfullMigrationMessage(const FText& InMessage);
+	UE_API void AddSucessfullMigrationMessage(const FText& InMessage);
 
 	// Notify the user that some asset may require some extra work to fully migrate
-	void AddWarningMigrationMessage(const FText& InMessage);
+	UE_API void AddWarningMigrationMessage(const FText& InMessage);
 
 	// Notify the user that some asset couldn't be migrated
-	void AddErrorMigrationMessage(const FText& InMessage);
+	UE_API void AddErrorMigrationMessage(const FText& InMessage);
 
 	struct FScopedMountPoint;
-	const FScopedMountPoint& GetDestinationMountPoint() const;
+	UE_API const FScopedMountPoint& GetDestinationMountPoint() const;
 
 	struct FMigrationPackageData;
 	// Return the data associated to a package that is migrated
-	const TArray<FMigrationPackageData>& GetMigrationPackagesData() const;
+	UE_API const TArray<FMigrationPackageData>& GetMigrationPackagesData() const;
 
 	// Return the array of non instanced package where moved because they were in the way of the migration
-	const TArray<UPackage*>& GetMovedOutOfTheWayPackages() const;
+	UE_API const TArray<UPackage*>& GetMovedOutOfTheWayPackages() const;
 
 	// Return the array of package dependencies that where not migrated because there where excluded by some automated filters
-	const TArray<FString>& GetExcludedDependencies() const;
+	UE_API const TArray<FString>& GetExcludedDependencies() const;
 
 	enum class EPackageMigrationStep
 	{
@@ -72,7 +74,7 @@ struct ASSETTOOLS_API FPackageMigrationContext : public FGCObject
 	};
 
 	// Utility to manage the temp mount point if needed and to help convert the paths during the migration
-	struct ASSETTOOLS_API FScopedMountPoint
+	struct FScopedMountPoint
 	{
 	public:
 		FScopedMountPoint() = delete;
@@ -82,21 +84,21 @@ struct ASSETTOOLS_API FPackageMigrationContext : public FGCObject
 		FScopedMountPoint(FScopedMountPoint&&) = default;
 		FScopedMountPoint& operator=(FScopedMountPoint&&) = default;
 
-		const FString& GetRootPath() const;
-		const FString& GetContentPath() const;
+		UE_API const FString& GetRootPath() const;
+		UE_API const FString& GetContentPath() const;
 
-		FString GetMigratedPackageFilename(const FString& LongPackageName, const FStringView& InExtension) const;
+		UE_API FString GetMigratedPackageFilename(const FString& LongPackageName, const FStringView& InExtension) const;
 
 		// Get name of package when moved under the destination mount point
-		FString GetNewPackageNameForMigration(const FString& LongPackageName) const;
+		UE_API FString GetNewPackageNameForMigration(const FString& LongPackageName) const;
 
 	private:
 		friend UAssetToolsImpl;
 		friend FPackageMigrationContext;
 		friend UE::AssetTools::Private::FPackageMigrationImpl;
 
-		FScopedMountPoint(FString&& InRootPath, FString&& InContentPath);
-		~FScopedMountPoint();
+		UE_API FScopedMountPoint(FString&& InRootPath, FString&& InContentPath);
+		UE_API ~FScopedMountPoint();
 
 		FString RootPath;
 		FString ContentPath;
@@ -104,13 +106,13 @@ struct ASSETTOOLS_API FPackageMigrationContext : public FGCObject
 	};
 
 	// The data associated to an package that take part of the migration process.
-	struct ASSETTOOLS_API FMigrationPackageData
+	struct FMigrationPackageData
 	{
-		const FString& GetInstancedPackageName() const;
-		const FString& GetOriginalPackageName() const;
-		const FString& GetDestinationFilename() const;
+		UE_API const FString& GetInstancedPackageName() const;
+		UE_API const FString& GetOriginalPackageName() const;
+		UE_API const FString& GetDestinationFilename() const;
 
-		UPackage* GetInstancedPackage() const;
+		UE_API UPackage* GetInstancedPackage() const;
 
 		// Tell the migration that package must be loaded
 		bool bNeedInstancedLoad = false;
@@ -118,7 +120,7 @@ struct ASSETTOOLS_API FPackageMigrationContext : public FGCObject
 		// Tell the migration that the package must be migrated by save
 		bool bNeedToBeSaveMigrated = false;
 
-		FMigrationPackageData(const FString& InInstancedPackageName, const FString& InOriginalPackageName, const FString& InDestinationFilename);
+		UE_API FMigrationPackageData(const FString& InInstancedPackageName, const FString& InOriginalPackageName, const FString& InDestinationFilename);
 
 	private:
 		friend UAssetToolsImpl;
@@ -138,9 +140,9 @@ struct ASSETTOOLS_API FPackageMigrationContext : public FGCObject
 	FPackageMigrationContext operator=(FPackageMigrationContext&&) = delete;
 
 
-	virtual void AddReferencedObjects( FReferenceCollector& Collector ) override;
+	UE_API virtual void AddReferencedObjects( FReferenceCollector& Collector ) override;
 
-	virtual FString GetReferencerName() const override;
+	UE_API virtual FString GetReferencerName() const override;
 
 private:
 
@@ -148,19 +150,19 @@ private:
 	friend UE::AssetTools::Private::FPackageMigrationImpl;
 
 	// Move a package out of the way for the duration of the migration
-	void MoveInTheWayPackage(UPackage* Package);
+	UE_API void MoveInTheWayPackage(UPackage* Package);
 
-	FPackageMigrationContext(FScopedMountPoint&& InMountPoint);
+	UE_API FPackageMigrationContext(FScopedMountPoint&& InMountPoint);
 
 	~FPackageMigrationContext() = default;
 
 	// Helper to move existing package out of the way during the migration
-	struct ASSETTOOLS_API FScopedTemporalyMovedPackage
+	struct FScopedTemporalyMovedPackage
 	{
 	public:
-		FScopedTemporalyMovedPackage(UPackage* InPackageToMove);
+		UE_API FScopedTemporalyMovedPackage(UPackage* InPackageToMove);
 
-		~FScopedTemporalyMovedPackage();
+		UE_API ~FScopedTemporalyMovedPackage();
 
 		TObjectPtr<UPackage> PackageToMove = nullptr;
 		FString OriginalName;
@@ -192,4 +194,6 @@ private:
 	TArray<FString> ExcludedDependencies;
 };
 
-}
\ No newline at end of file
+}
+
+#undef UE_API
diff --git a/Engine/Source/Developer/AudioSettingsEditor/Public/TargetPlatformAudioCustomization.h b/Engine/Source/Developer/AudioSettingsEditor/Public/TargetPlatformAudioCustomization.h
--- a/Engine/Source/Developer/AudioSettingsEditor/Public/TargetPlatformAudioCustomization.h
+++ b/Engine/Source/Developer/AudioSettingsEditor/Public/TargetPlatformAudioCustomization.h
@@ -8,25 +8,27 @@
 
 #include "AudioPluginUtilities.h"
 
-class AUDIOSETTINGSEDITOR_API FAudioPluginWidgetManager
+#define UE_API AUDIOSETTINGSEDITOR_API
+
+class FAudioPluginWidgetManager
 {
 public:
-	FAudioPluginWidgetManager();
-	~FAudioPluginWidgetManager();
+	UE_API FAudioPluginWidgetManager();
+	UE_API ~FAudioPluginWidgetManager();
 
 	/* Builds out the audio category for a specific audio section for a platform settings page. */
-	void BuildAudioCategory(IDetailLayoutBuilder& DetailLayout, const FString& PlatformName, const UStruct* ClassOuterMost = nullptr);
+	UE_API void BuildAudioCategory(IDetailLayoutBuilder& DetailLayout, const FString& PlatformName, const UStruct* ClassOuterMost = nullptr);
 
 	/** Creates widget from a scan of loaded audio plugins for an individual plugin type. */
-	TSharedRef<SWidget> MakeAudioPluginSelectorWidget(const TSharedPtr<IPropertyHandle>& PropertyHandle, EAudioPlugin AudioPluginType, const FString& PlatformName);
+	UE_API TSharedRef<SWidget> MakeAudioPluginSelectorWidget(const TSharedPtr<IPropertyHandle>& PropertyHandle, EAudioPlugin AudioPluginType, const FString& PlatformName);
 
 private:
 	/** Handles when a new plugin is selected. */
-	static void OnPluginSelected(FString PluginName, TSharedPtr<IPropertyHandle> PropertyHandle);
+	static UE_API void OnPluginSelected(FString PluginName, TSharedPtr<IPropertyHandle> PropertyHandle);
 
-	void OnPluginTextCommitted(const FText& InText, ETextCommit::Type CommitType, EAudioPlugin AudioPluginType, TSharedPtr<IPropertyHandle> PropertyHandle);
+	UE_API void OnPluginTextCommitted(const FText& InText, ETextCommit::Type CommitType, EAudioPlugin AudioPluginType, TSharedPtr<IPropertyHandle> PropertyHandle);
 
-	FText OnGetPluginText(EAudioPlugin AudioPluginType);
+	UE_API FText OnGetPluginText(EAudioPlugin AudioPluginType);
 
 private:
 	/** Cached references to text for Spatialization, Reverb and Occlusion settings */
@@ -45,3 +47,5 @@ private:
 	TArray<TSharedPtr<FText>> ReverbPlugins;
 	TArray<TSharedPtr<FText>> OcclusionPlugins;
 };
+
+#undef UE_API
diff --git a/Engine/Source/Developer/AutomationDriver/Public/IAutomationDriverModule.h b/Engine/Source/Developer/AutomationDriver/Public/IAutomationDriverModule.h
--- a/Engine/Source/Developer/AutomationDriver/Public/IAutomationDriverModule.h
+++ b/Engine/Source/Developer/AutomationDriver/Public/IAutomationDriverModule.h
@@ -8,7 +8,7 @@ class IAutomationDriver;
 class IAsyncAutomationDriver;
 class FDriverConfiguration;
 
-class AUTOMATIONDRIVER_API IAutomationDriverModule
+class IAutomationDriverModule
 	: public IModuleInterface
 {
 public:
diff --git a/Engine/Source/Developer/AutomationDriver/Public/LocateBy.h b/Engine/Source/Developer/AutomationDriver/Public/LocateBy.h
--- a/Engine/Source/Developer/AutomationDriver/Public/LocateBy.h
+++ b/Engine/Source/Developer/AutomationDriver/Public/LocateBy.h
@@ -5,6 +5,8 @@
 #include "Widgets/SWidget.h"
 #include "AutomationDriverTypeDefs.h"
 
+#define UE_API AUTOMATIONDRIVER_API
+
 class IElementLocator;
 
 DECLARE_DELEGATE_OneParam(FLocateSlateWidgetElementDelegate, TArray<TSharedRef<SWidget>>& /*OutWidgets*/);
@@ -13,7 +15,7 @@ DECLARE_DELEGATE_OneParam(FLocateSlateWidgetPathElementDelegate, TArray<FWidgetP
 /**
  * Represents a collection of fluent helper functions designed to make accessing and creating element locators easy
  */
-class AUTOMATIONDRIVER_API By
+class By
 {
 public:
 
@@ -24,7 +26,7 @@ public:
 	 * @param DebugName - The string to display for errors when no elements can be located
 	 * @return a locator which uses the specified delegate to discover appropriate elements
 	 */
-	static TSharedRef<IElementLocator, ESPMode::ThreadSafe> Delegate(const FLocateSlateWidgetElementDelegate& Value, FStringView DebugName = { });
+	static UE_API TSharedRef<IElementLocator, ESPMode::ThreadSafe> Delegate(const FLocateSlateWidgetElementDelegate& Value, FStringView DebugName = { });
 
 	/**
 	 * Creates a new element locator that exposes the collection of FWidgetPaths returned from the FLocateSlateWidgetElementDelegate
@@ -33,7 +35,7 @@ public:
 	 * @param DebugName - The string to display for errors when no elements can be located
 	 * @return a locator which uses the specified delegate to discover appropriate elements
 	 */
-	static TSharedRef<IElementLocator, ESPMode::ThreadSafe> Delegate(const FLocateSlateWidgetPathElementDelegate& Value, FStringView DebugName = { });
+	static UE_API TSharedRef<IElementLocator, ESPMode::ThreadSafe> Delegate(const FLocateSlateWidgetPathElementDelegate& Value, FStringView DebugName = { });
 
 	/**
 	 * Creates a new element locator that exposes the collection of SWidgets returned from the lambda
@@ -42,7 +44,7 @@ public:
 	 * @param DebugName - The string to display for errors when no elements can be located
 	 * @return a locator which uses the specified lambda to discover appropriate elements
 	 */
-	static TSharedRef<IElementLocator, ESPMode::ThreadSafe> WidgetLambda(const TFunction<void(TArray<TSharedRef<SWidget>>&)>& Value, FStringView DebugName = { });
+	static UE_API TSharedRef<IElementLocator, ESPMode::ThreadSafe> WidgetLambda(const TFunction<void(TArray<TSharedRef<SWidget>>&)>& Value, FStringView DebugName = { });
 
 	/**
 	 * Creates a new element locator that exposes the collection of FWidgetPaths returned from the lambda
@@ -51,14 +53,14 @@ public:
 	 * @param DebugName - The string to display for errors when no elements can be located
 	 * @return a locator which uses the specified lambda to discover appropriate elements
 	 */
-	static TSharedRef<IElementLocator, ESPMode::ThreadSafe> WidgetPathLambda(const TFunction<void(TArray<FWidgetPath>&)>& Value, FStringView DebugName = { });
+	static UE_API TSharedRef<IElementLocator, ESPMode::ThreadSafe> WidgetPathLambda(const TFunction<void(TArray<FWidgetPath>&)>& Value, FStringView DebugName = { });
 
 	/**
 	 * Creates a new element locator that limits its discovered elements to ones with the specified Id
 	 * @param Value - The MetaData ID of the element to find
 	 * @return a locator which uses the specified Id to discover appropriate elements
 	 */
-	static TSharedRef<IElementLocator, ESPMode::ThreadSafe> Id(const FString& Value);
+	static UE_API TSharedRef<IElementLocator, ESPMode::ThreadSafe> Id(const FString& Value);
 
 	/**
 	 * Creates a new element locator that limits its discovered elements to ones with the specified Id,
@@ -67,14 +69,14 @@ public:
 	 * @param Value - The MetaData ID of the element to find
 	 * @return a locator which uses the specified Id to discover appropriate elements
 	 */
-	static TSharedRef<IElementLocator, ESPMode::ThreadSafe> Id(const FDriverElementRef& Root, const FString& Value);
+	static UE_API TSharedRef<IElementLocator, ESPMode::ThreadSafe> Id(const FDriverElementRef& Root, const FString& Value);
 
 	/**
 	 * Creates a new element locator that limits its discovered elements to ones with the specified Id
 	 * @param Value - The MetaData ID of the element to find
 	 * @return a locator which uses the specified Id to discover appropriate elements
 	 */
-	static TSharedRef<IElementLocator, ESPMode::ThreadSafe> Id(const FName& Value);
+	static UE_API TSharedRef<IElementLocator, ESPMode::ThreadSafe> Id(const FName& Value);
 	
 	/**
 	 * Creates a new element locator that limits its discovered elements to ones with the specified Id,
@@ -83,7 +85,7 @@ public:
 	 * @param Value - The MetaData ID of the element to find
 	 * @return a locator which uses the specified Id to discover appropriate elements
 	 */
-	static TSharedRef<IElementLocator, ESPMode::ThreadSafe> Id(const FDriverElementRef& Root, const FName& Value);
+	static UE_API TSharedRef<IElementLocator, ESPMode::ThreadSafe> Id(const FDriverElementRef& Root, const FName& Value);
 
 
 	/**
@@ -91,7 +93,7 @@ public:
 	 * @param Value - The MetaData ID of the element to find
 	 * @return a locator which uses the specified Id to discover appropriate elements
 	 */
-	static TSharedRef<IElementLocator, ESPMode::ThreadSafe> Id(const TCHAR* Value);
+	static UE_API TSharedRef<IElementLocator, ESPMode::ThreadSafe> Id(const TCHAR* Value);
 
 	/**
 	 * Creates a new element locator that limits its discovered elements to ones with the specified Id,
@@ -100,7 +102,7 @@ public:
 	 * @param Value - The MetaData ID of the element to find
 	 * @return a locator which uses the specified Id to discover appropriate elements
 	 */
-	static TSharedRef<IElementLocator, ESPMode::ThreadSafe> Id(const FDriverElementRef& Root, const TCHAR* Value);
+	static UE_API TSharedRef<IElementLocator, ESPMode::ThreadSafe> Id(const FDriverElementRef& Root, const TCHAR* Value);
 
 
 	/**
@@ -108,7 +110,7 @@ public:
 	 * @param Value - The MetaData ID of the element to find
 	 * @return a locator which uses the specified Id to discover appropriate elements
 	 */
-	static TSharedRef<IElementLocator, ESPMode::ThreadSafe> Id(const char* Value);
+	static UE_API TSharedRef<IElementLocator, ESPMode::ThreadSafe> Id(const char* Value);
 
 	/**
 	 * Creates a new element locator that limits its discovered elements to ones with the specified Id,
@@ -117,7 +119,7 @@ public:
 	 * @param Value - The MetaData ID of the element to find
 	 * @return a locator which uses the specified Id to discover appropriate elements
 	 */
-	static TSharedRef<IElementLocator, ESPMode::ThreadSafe> Id(const FDriverElementRef& Root, const char* Value);
+	static UE_API TSharedRef<IElementLocator, ESPMode::ThreadSafe> Id(const FDriverElementRef& Root, const char* Value);
 
 	/**
 	 * Creates a new element locator that limits its discovered elements to ones matching the specified "path"
@@ -140,7 +142,7 @@ public:
 	 * @param Value - The path to use
 	 * @return a locator which uses the specified path to discover appropriate elements
 	 */
-	static TSharedRef<IElementLocator, ESPMode::ThreadSafe> Path(const FString& Value);
+	static UE_API TSharedRef<IElementLocator, ESPMode::ThreadSafe> Path(const FString& Value);
 
 	/**
 	 * Creates a new element locator that limits its discovered elements to ones matching the specified "path",
@@ -165,7 +167,7 @@ public:
 	 * @param Value - The path to use
 	 * @return a locator which uses the specified path to discover appropriate elements
 	 */
-	static TSharedRef<IElementLocator, ESPMode::ThreadSafe> Path(const FDriverElementRef& Root, const FString& Value);
+	static UE_API TSharedRef<IElementLocator, ESPMode::ThreadSafe> Path(const FDriverElementRef& Root, const FString& Value);
 
 	
 	/**
@@ -189,7 +191,7 @@ public:
 	 * @param Value - The path to use
 	 * @return a locator which uses the specified path to discover appropriate elements
 	 */
-	static TSharedRef<IElementLocator, ESPMode::ThreadSafe> Path(const FName& Value);
+	static UE_API TSharedRef<IElementLocator, ESPMode::ThreadSafe> Path(const FName& Value);
 
 	/**
 	 * Creates a new element locator that limits its discovered elements to ones matching the specified "path",
@@ -214,7 +216,7 @@ public:
 	 * @param Value - The path to use
 	 * @return a locator which uses the specified path to discover appropriate elements
 	 */
-	static TSharedRef<IElementLocator, ESPMode::ThreadSafe> Path(const FDriverElementRef& Root, const FName& Value);
+	static UE_API TSharedRef<IElementLocator, ESPMode::ThreadSafe> Path(const FDriverElementRef& Root, const FName& Value);
 
 	/**
 	 * Creates a new element locator that limits its discovered elements to ones matching the specified "path"
@@ -237,7 +239,7 @@ public:
 	 * @param Value - The path to use
 	 * @return a locator which uses the specified path to discover appropriate elements
 	 */
-	static TSharedRef<IElementLocator, ESPMode::ThreadSafe> Path(const TCHAR* Value);
+	static UE_API TSharedRef<IElementLocator, ESPMode::ThreadSafe> Path(const TCHAR* Value);
 
 	/**
 	 * Creates a new element locator that limits its discovered elements to ones matching the specified "path",
@@ -262,7 +264,7 @@ public:
 	 * @param Value - The path to use
 	 * @return a locator which uses the specified path to discover appropriate elements
 	 */
-	static TSharedRef<IElementLocator, ESPMode::ThreadSafe> Path(const FDriverElementRef& Root, const TCHAR* Value);
+	static UE_API TSharedRef<IElementLocator, ESPMode::ThreadSafe> Path(const FDriverElementRef& Root, const TCHAR* Value);
 
 	/**
 	 * Creates a new element locator that limits its discovered elements to ones matching the specified "path"
@@ -285,7 +287,7 @@ public:
 	 * @param Value - The path to use
 	 * @return a locator which uses the specified path to discover appropriate elements
 	 */
-	static TSharedRef<IElementLocator, ESPMode::ThreadSafe> Path(const char* Value);
+	static UE_API TSharedRef<IElementLocator, ESPMode::ThreadSafe> Path(const char* Value);
 
 	/**
 	 * Creates a new element locator that limits its discovered elements to ones matching the specified "path",
@@ -310,27 +312,27 @@ public:
 	 * @param Value - The path to use
 	 * @return a locator which uses the specified path to discover appropriate elements
 	 */
-	static TSharedRef<IElementLocator, ESPMode::ThreadSafe> Path(const FDriverElementRef& Root, const char* Value);
+	static UE_API TSharedRef<IElementLocator, ESPMode::ThreadSafe> Path(const FDriverElementRef& Root, const char* Value);
 
 	/**
 	 * Creates a new element locator that limits its discovered elements to the one under the cursors current position
 	 * @return a locator which uses the cursor position to discover appropriate elements
 	 */
-	static TSharedRef<IElementLocator, ESPMode::ThreadSafe> Cursor();
+	static UE_API TSharedRef<IElementLocator, ESPMode::ThreadSafe> Cursor();
 	
 	/**
 	 * Creates a new element locator that limits its discovered element to the one that's currently in keyboard focus
 	 * @return a locator which uses the keyboard focus to discover the appropriate element
 	 */
-	static TSharedRef<IElementLocator, ESPMode::ThreadSafe> KeyboardFocus();
+	static UE_API TSharedRef<IElementLocator, ESPMode::ThreadSafe> KeyboardFocus();
 
 	/**
 	 * Creates a new element locator that limits its discovered element to the one that's currently in user focus
 	 * @return a locator which uses the user focus to discover the appropriate element
 	 */
-	static TSharedRef<IElementLocator, ESPMode::ThreadSafe> UserFocus(uint32 UserIndex = 0);
+	static UE_API TSharedRef<IElementLocator, ESPMode::ThreadSafe> UserFocus(uint32 UserIndex = 0);
 
-	class AUTOMATIONDRIVER_API TextFilter
+	class TextFilter
 	{
 	public:
 
@@ -344,7 +346,7 @@ public:
 		 * @param SearchDir - Indicates whether the search starts at the beginning or at the end ( defaults to ESearchDir::FromStart )
 		 * @return a locator which uses the specified root locator and a text string to discover appropriate elements
 		 */
-		static TSharedRef<IElementLocator, ESPMode::ThreadSafe> Contains(
+		static UE_API TSharedRef<IElementLocator, ESPMode::ThreadSafe> Contains(
 			const FElementLocatorRef& RootLocator,
 			const FString& Value,
 			ESearchCase::Type SearchCase = ESearchCase::IgnoreCase,
@@ -360,10 +362,12 @@ public:
 		 * @param SearchCase - Indicates whether the search is case sensitive or not ( defaults to ESearchCase::CaseSensitive )
 		 * @return a locator which uses the specified root locator and a text string to discover appropriate elements
 		 */
-		static TSharedRef<IElementLocator, ESPMode::ThreadSafe> Equals(
+		static UE_API TSharedRef<IElementLocator, ESPMode::ThreadSafe> Equals(
 			const FElementLocatorRef& RootLocator,
 			const FString& Value,
 			ESearchCase::Type SearchCase = ESearchCase::CaseSensitive
 		);
 	};
 };
+
+#undef UE_API
diff --git a/Engine/Source/Developer/AutomationDriver/Public/WaitUntil.h b/Engine/Source/Developer/AutomationDriver/Public/WaitUntil.h
--- a/Engine/Source/Developer/AutomationDriver/Public/WaitUntil.h
+++ b/Engine/Source/Developer/AutomationDriver/Public/WaitUntil.h
@@ -8,12 +8,14 @@
 #include "Templates/Function.h"
 #include "Templates/SharedPointer.h"
 
+#define UE_API AUTOMATIONDRIVER_API
+
 class IElementLocator;
 
 /**
  * Represents the state of an active wait action for the driver
  */
-struct AUTOMATIONDRIVER_API FDriverWaitResponse
+struct FDriverWaitResponse
 {
 public:
 
@@ -27,24 +29,24 @@ public:
 	/**
 	 * @return a FDriverWaitResponse with a state of PASSED and a wait of zero
 	 */
-	static FDriverWaitResponse Passed();
+	static UE_API FDriverWaitResponse Passed();
 
 	/**
 	 * @return a FDriverWaitResponse with a state of WAIT and a wait of 0.5 seconds
 	 */
-	static FDriverWaitResponse Wait();
+	static UE_API FDriverWaitResponse Wait();
 
 	/**
 	 * @return a FDriverWaitResponse with a state of WAIT and a wait of the specified timespan
 	 */
-	static FDriverWaitResponse Wait(FTimespan Timespan);
+	static UE_API FDriverWaitResponse Wait(FTimespan Timespan);
 
 	/**
 	 * @return a FDriverWaitResponse with a state of FAILED and a wait of zero
 	 */
-	static FDriverWaitResponse Failed();
+	static UE_API FDriverWaitResponse Failed();
 
-	FDriverWaitResponse(EState InState, FTimespan InNextWait);
+	UE_API FDriverWaitResponse(EState InState, FTimespan InNextWait);
 
 	// How long the driver should wait before re-evaluating the wait condition again
 	const FTimespan NextWait;
@@ -59,16 +61,16 @@ DECLARE_DELEGATE_RetVal(bool, FDriverWaitConditionDelegate);
 /**
  * A fluent wrapper around timespan to enforce obvious differences between specified Timeout and Interval values for waits
  */
-class AUTOMATIONDRIVER_API FWaitTimeout
+class FWaitTimeout
 {
 public:
 
-	static FWaitTimeout InMilliseconds(double Value);
-	static FWaitTimeout InSeconds(double Value);
-	static FWaitTimeout InMinutes(double Value);
-	static FWaitTimeout InHours(double Value);
+	static UE_API FWaitTimeout InMilliseconds(double Value);
+	static UE_API FWaitTimeout InSeconds(double Value);
+	static UE_API FWaitTimeout InMinutes(double Value);
+	static UE_API FWaitTimeout InHours(double Value);
 
-	explicit FWaitTimeout(FTimespan InTimespan);
+	UE_API explicit FWaitTimeout(FTimespan InTimespan);
 
 	const FTimespan Timespan;
 };
@@ -76,16 +78,16 @@ public:
 /**
  * A fluent wrapper around timespan to enforce obvious differences between specified Timeout and Interval values for waits
  */
-class AUTOMATIONDRIVER_API FWaitInterval
+class FWaitInterval
 {
 public:
 
-	static FWaitInterval InMilliseconds(double Value);
-	static FWaitInterval InSeconds(double Value);
-	static FWaitInterval InMinutes(double Value);
-	static FWaitInterval InHours(double Value);
+	static UE_API FWaitInterval InMilliseconds(double Value);
+	static UE_API FWaitInterval InSeconds(double Value);
+	static UE_API FWaitInterval InMinutes(double Value);
+	static UE_API FWaitInterval InHours(double Value);
 
-	explicit FWaitInterval(FTimespan InTimespan);
+	UE_API explicit FWaitInterval(FTimespan InTimespan);
 
 	const FTimespan Timespan;
 };
@@ -93,126 +95,128 @@ public:
 /**
  * Represents a collection of fluent helper functions designed to make accessing and creating driver wait delegates easier
  */
-class AUTOMATIONDRIVER_API Until
+class Until
 {
 public:
 
 	/**
 	 * Creates a new wait delegate which completes it's wait only if the specified element locator discovers elements or if the specified timeout timespan elapses
 	 */
-	static FDriverWaitDelegate ElementExists(const TSharedRef<IElementLocator, ESPMode::ThreadSafe>& ElementLocator, FWaitTimeout Timeout);
+	static UE_API FDriverWaitDelegate ElementExists(const TSharedRef<IElementLocator, ESPMode::ThreadSafe>& ElementLocator, FWaitTimeout Timeout);
 
 	/**
 	 * Creates a new wait delegate which completes it's wait only if the specified element locator discovers elements or if the specified timeout timespan elapses;
 	 * The element locator is only re-evaluated at the specified wait interval
 	 */
-	static FDriverWaitDelegate ElementExists(const TSharedRef<IElementLocator, ESPMode::ThreadSafe>& ElementLocator, FWaitInterval Interval, FWaitTimeout Timeout);
+	static UE_API FDriverWaitDelegate ElementExists(const TSharedRef<IElementLocator, ESPMode::ThreadSafe>& ElementLocator, FWaitInterval Interval, FWaitTimeout Timeout);
 
 	/**
 	 * Creates a new wait delegate which completes it's wait only if the specified element locator discovers visible elements or if the specified timeout timespan elapses
 	 */
-	static FDriverWaitDelegate ElementIsVisible(const TSharedRef<IElementLocator, ESPMode::ThreadSafe>& ElementLocator, FWaitTimeout Timeout);
+	static UE_API FDriverWaitDelegate ElementIsVisible(const TSharedRef<IElementLocator, ESPMode::ThreadSafe>& ElementLocator, FWaitTimeout Timeout);
 
 	/**
 	 * Creates a new wait delegate which completes it's wait only if the specified element locator discovers visible elements or if the specified timeout timespan elapses;
 	 * The element locator is only re-evaluated at the specified wait interval
 	 */
-	static FDriverWaitDelegate ElementIsVisible(const TSharedRef<IElementLocator, ESPMode::ThreadSafe>& ElementLocator, FWaitInterval Interval, FWaitTimeout Timeout);
+	static UE_API FDriverWaitDelegate ElementIsVisible(const TSharedRef<IElementLocator, ESPMode::ThreadSafe>& ElementLocator, FWaitInterval Interval, FWaitTimeout Timeout);
 
 	/**
 	* Creates a new wait delegate which completes it's wait only if the specified element locator discovers hidden elements or if the specified timeout timespan elapses
 	*/
-	static FDriverWaitDelegate ElementIsHidden(const TSharedRef<IElementLocator, ESPMode::ThreadSafe>& ElementLocator, FWaitTimeout Timeout);
+	static UE_API FDriverWaitDelegate ElementIsHidden(const TSharedRef<IElementLocator, ESPMode::ThreadSafe>& ElementLocator, FWaitTimeout Timeout);
 
 	/**
 	 * Creates a new wait delegate which completes it's wait only if the specified element locator discovers hidden elements or if the specified timeout timespan elapses;
 	 * The element locator is only re-evaluated at the specified wait interval
 	 */
-	static FDriverWaitDelegate ElementIsHidden(const TSharedRef<IElementLocator, ESPMode::ThreadSafe>& ElementLocator, FWaitInterval Interval, FWaitTimeout Timeout);
+	static UE_API FDriverWaitDelegate ElementIsHidden(const TSharedRef<IElementLocator, ESPMode::ThreadSafe>& ElementLocator, FWaitInterval Interval, FWaitTimeout Timeout);
 
 	/**
 	 * Creates a new wait delegate which completes it's wait only if the specified element locator discovers interactable elements or if the specified timeout timespan elapses
 	 */
-	static FDriverWaitDelegate ElementIsInteractable(const TSharedRef<IElementLocator, ESPMode::ThreadSafe>& ElementLocator, FWaitTimeout Timeout);
+	static UE_API FDriverWaitDelegate ElementIsInteractable(const TSharedRef<IElementLocator, ESPMode::ThreadSafe>& ElementLocator, FWaitTimeout Timeout);
 
 	/**
 	 * Creates a new wait delegate which completes it's wait only if the specified element locator discovers interactable elements or if the specified timeout timespan elapses;
 	 * The element locator is only re-evaluated at the specified wait interval
 	 */
-	static FDriverWaitDelegate ElementIsInteractable(const TSharedRef<IElementLocator, ESPMode::ThreadSafe>& ElementLocator, FWaitInterval Interval, FWaitTimeout Timeout);
+	static UE_API FDriverWaitDelegate ElementIsInteractable(const TSharedRef<IElementLocator, ESPMode::ThreadSafe>& ElementLocator, FWaitInterval Interval, FWaitTimeout Timeout);
 
 	/**
 	 * Creates a new wait delegate which completes it's wait only if the specified element locator discovers a scrollable element whose scroll position is at the beginning
 	 * or if the specified timeout timespan elapses
 	 */
-	static FDriverWaitDelegate ElementIsScrolledToBeginning(const TSharedRef<IElementLocator, ESPMode::ThreadSafe>& ElementLocator, FWaitTimeout Timeout);
+	static UE_API FDriverWaitDelegate ElementIsScrolledToBeginning(const TSharedRef<IElementLocator, ESPMode::ThreadSafe>& ElementLocator, FWaitTimeout Timeout);
 
 	/**
 	 * Creates a new wait delegate which completes it's wait only if the specified element locator discovers a scrollable element whose scroll position is at the beginning
 	 * or if the specified timeout timespan elapses
 	 * The element locator is only re-evaluated at the specified wait interval
 	 */
-	static FDriverWaitDelegate ElementIsScrolledToBeginning(const TSharedRef<IElementLocator, ESPMode::ThreadSafe>& ElementLocator, FWaitInterval Interval, FWaitTimeout Timeout);
+	static UE_API FDriverWaitDelegate ElementIsScrolledToBeginning(const TSharedRef<IElementLocator, ESPMode::ThreadSafe>& ElementLocator, FWaitInterval Interval, FWaitTimeout Timeout);
 
 	/**
 	 * Creates a new wait delegate which completes it's wait only if the specified element locator discovers a scrollable element whose scroll position is at the end
 	 * or if the specified timeout timespan elapses
 	 */
-	static FDriverWaitDelegate ElementIsScrolledToEnd(const TSharedRef<IElementLocator, ESPMode::ThreadSafe>& ElementLocator, FWaitTimeout Timeout);
+	static UE_API FDriverWaitDelegate ElementIsScrolledToEnd(const TSharedRef<IElementLocator, ESPMode::ThreadSafe>& ElementLocator, FWaitTimeout Timeout);
 	
 	/**
 	 * Creates a new wait delegate which completes it's wait only if the specified element locator discovers a scrollable element whose scroll position is at the end
 	 * or if the specified timeout timespan elapses
 	 * The element locator is only re-evaluated at the specified wait interval
 	 */
-	static FDriverWaitDelegate ElementIsScrolledToEnd(const TSharedRef<IElementLocator, ESPMode::ThreadSafe>& ElementLocator, FWaitInterval Interval, FWaitTimeout Timeout);
+	static UE_API FDriverWaitDelegate ElementIsScrolledToEnd(const TSharedRef<IElementLocator, ESPMode::ThreadSafe>& ElementLocator, FWaitInterval Interval, FWaitTimeout Timeout);
 
 	/**
 	 * Creates a new wait delegate which completes its wait only if the specified element locator discovers an element focused by the default user or if the specified timeout timespan elapses
 	 */
-	static FDriverWaitDelegate ElementIsFocusedByKeyboard(const TSharedRef<IElementLocator, ESPMode::ThreadSafe>& ElementLocator, FWaitTimeout Timeout);
+	static UE_API FDriverWaitDelegate ElementIsFocusedByKeyboard(const TSharedRef<IElementLocator, ESPMode::ThreadSafe>& ElementLocator, FWaitTimeout Timeout);
 
 	/**
 	 * Creates a new wait delegate which completes it's wait only if the specified element locator discovers an element focused by the default user or if the specified timeout timespan elapses;
 	 * The element locator is only re-evaluated at the specified wait interval
 	 */
-	static FDriverWaitDelegate ElementIsFocusedByKeyboard(const TSharedRef<IElementLocator, ESPMode::ThreadSafe>& ElementLocator, FWaitInterval Interval, FWaitTimeout Timeout);
+	static UE_API FDriverWaitDelegate ElementIsFocusedByKeyboard(const TSharedRef<IElementLocator, ESPMode::ThreadSafe>& ElementLocator, FWaitInterval Interval, FWaitTimeout Timeout);
 
 	/**
 	 * Creates a new wait delegate which completes its wait only if the specified element locator discovers an element focused by the specified user or if the specified timeout timespan elapses
 	 */
-	static FDriverWaitDelegate ElementIsFocusedByUser(uint32 UserIndex, const TSharedRef<IElementLocator, ESPMode::ThreadSafe>& ElementLocator, FWaitTimeout Timeout);
+	static UE_API FDriverWaitDelegate ElementIsFocusedByUser(uint32 UserIndex, const TSharedRef<IElementLocator, ESPMode::ThreadSafe>& ElementLocator, FWaitTimeout Timeout);
 
 	/**
 	 * Creates a new wait delegate which completes it's wait only if the specified element locator discovers an element focused by the specified user or if the specified timeout timespan elapses;
 	 * The element locator is only re-evaluated at the specified wait interval
 	 */
-	static FDriverWaitDelegate ElementIsFocusedByUser(uint32 UserIndex, const TSharedRef<IElementLocator, ESPMode::ThreadSafe>& ElementLocator, FWaitInterval Interval, FWaitTimeout Timeout);
+	static UE_API FDriverWaitDelegate ElementIsFocusedByUser(uint32 UserIndex, const TSharedRef<IElementLocator, ESPMode::ThreadSafe>& ElementLocator, FWaitInterval Interval, FWaitTimeout Timeout);
 
 	/**
 	 * Creates a new wait delegate which completes it's wait only if the specified condition returns true or if the specified timeout timespan elapses
 	 */
-	static FDriverWaitDelegate Condition(const TFunction<bool()>& Function, FWaitTimeout Timeout);
+	static UE_API FDriverWaitDelegate Condition(const TFunction<bool()>& Function, FWaitTimeout Timeout);
 
 	/**
 	 * Creates a new wait delegate which completes it's wait only if the specified condition returns true or if the specified timeout timespan elapses
 	 * The lambda is only re-evaluated at the specified wait interval
 	 */
-	static FDriverWaitDelegate Condition(const TFunction<bool()>& Function, FWaitInterval Interval, FWaitTimeout Timeout);
+	static UE_API FDriverWaitDelegate Condition(const TFunction<bool()>& Function, FWaitInterval Interval, FWaitTimeout Timeout);
 
 	/**
 	 * Creates a new wait delegate which completes it's wait only if the specified condition returns true or if the specified timeout timespan elapses
 	 */
-	static FDriverWaitDelegate Condition(const FDriverWaitConditionDelegate& Delegate, FWaitTimeout Timeout);
+	static UE_API FDriverWaitDelegate Condition(const FDriverWaitConditionDelegate& Delegate, FWaitTimeout Timeout);
 
 	/**
 	 * Creates a new wait delegate which completes it's wait only if the specified condition returns true or if the specified timeout timespan elapses
 	 * The delegate is only re-evaluated at the specified wait interval
 	 */
-	static FDriverWaitDelegate Condition(const FDriverWaitConditionDelegate& Delegate, FWaitInterval Interval, FWaitTimeout Timeout);
+	static UE_API FDriverWaitDelegate Condition(const FDriverWaitConditionDelegate& Delegate, FWaitInterval Interval, FWaitTimeout Timeout);
 
 	/**
 	 * Creates a new wait delegate which drives its state off the result of the specified lambda
 	 */
-	static FDriverWaitDelegate Lambda(const TFunction<FDriverWaitResponse(const FTimespan&)>& Value);
+	static UE_API FDriverWaitDelegate Lambda(const TFunction<FDriverWaitResponse(const FTimespan&)>& Value);
 };
+
+#undef UE_API
