If a top-level bracket is open with explicit intent not to transact any changes, ignore any transaction requests within the scope of the bracket
#tests controller automation tests

#Synced-CL 50041557

diff --git a/Engine/Plugins/Animation/AnimationData/Source/AnimationData/Private/AnimSequencerController.cpp b/Engine/Plugins/Animation/AnimationData/Source/AnimationData/Private/AnimSequencerController.cpp
--- a/Engine/Plugins/Animation/AnimationData/Source/AnimationData/Private/AnimSequencerController.cpp
+++ b/Engine/Plugins/Animation/AnimationData/Source/AnimationData/Private/AnimSequencerController.cpp
@@ -44,7 +44,12 @@ void UAnimSequencerController::OpenBracket(const FText& InTitle, bool bShouldTra
 {
 	ValidateModel();
 
-	if (UE::FChangeTransactor::CanTransactChanges() && !ChangeTransactor.IsTransactionPending() && bShouldTransact)
+	if (BracketDepth == 0)
+	{
+		bShouldTransactBracket = bShouldTransact;
+	}
+
+	if (UE::FChangeTransactor::CanTransactChanges() && !ChangeTransactor.IsTransactionPending() && ShouldTransact(bShouldTransact))
 	{
 		ChangeTransactor.OpenTransaction(InTitle);
 
@@ -78,7 +83,7 @@ void UAnimSequencerController::CloseBracket(bool bShouldTransact /*= true*/)
 
 	if (BracketDepth == 0)
 	{
-		if (UE::FChangeTransactor::CanTransactChanges() && bShouldTransact)
+		if (UE::FChangeTransactor::CanTransactChanges() && ShouldTransact(bShouldTransact && bShouldTransactBracket))
 		{
 			ensure(ChangeTransactor.IsTransactionPending());
 
@@ -90,6 +95,8 @@ void UAnimSequencerController::CloseBracket(bool bShouldTransact /*= true*/)
 		Model->GetNotifier().Notify(EAnimDataModelNotifyType::BracketClosed);
 		
 		Model->UnlockEvaluationAndModification();
+
+		bShouldTransactBracket = false;
 	}
 }
 
@@ -272,6 +279,11 @@ bool UAnimSequencerController::IgnoreSkeletonValidation() const
 	return Model->GetSkeleton() == nullptr && Model->bPopulated == false;
 }
 
+bool UAnimSequencerController::ShouldTransact(bool bInLocalShouldTransact) const
+{
+	return bInLocalShouldTransact && (BracketDepth == 0 || bShouldTransactBracket);
+}
+
 void UAnimSequencerController::ResizePlayLength(float Length, float T0, float T1, bool bShouldTransact)
 {
 	ResizeNumberOfFrames(ConvertSecondsToFrameNumber(Length), ConvertSecondsToFrameNumber(T0), ConvertSecondsToFrameNumber(T1), bShouldTransact);
@@ -371,7 +383,7 @@ bool UAnimSequencerController::RemoveBoneTracksMissingFromSkeleton(const USkelet
 				FBracket Bracket = ConditionalBracket(LOCTEXT("RemoveBoneTracksMissingFromSkeleton", "Validating Bone Animation Track Data against Skeleton"), bShouldTransact);
 				for (const FName& TrackName : TracksToBeRemoved)
 				{
-					RemoveBoneTrack(TrackName);
+					RemoveBoneTrack(TrackName, bShouldTransact);
 				}
 			}
 
@@ -420,7 +432,7 @@ void UAnimSequencerController::UpdateAttributesFromSkeleton(const USkeleton* Ske
 			FBracket Bracket = ConditionalBracket(LOCTEXT("VerifyAttributeBoneNames", "Remapping Animation Attribute Data"), bShouldTransact);
 			for (const FAnimationAttributeIdentifier& Identifier : ToRemoveIdentifiers)
 			{
-				RemoveAttribute(Identifier);
+				RemoveAttribute(Identifier, bShouldTransact);
 			}
 			
 			for (const TPair<FAnimationAttributeIdentifier, int32>& Pair : ToDuplicateIdentifiers)
@@ -428,8 +440,8 @@ void UAnimSequencerController::UpdateAttributesFromSkeleton(const USkeleton* Ske
 				const FAnimationAttributeIdentifier& DuplicateIdentifier = Pair.Key;								
 				FAnimationAttributeIdentifier NewIdentifier(DuplicateIdentifier.GetName(), Pair.Value, DuplicateIdentifier.GetBoneName(), DuplicateIdentifier.GetType());
 
-				DuplicateAttribute(Pair.Key, NewIdentifier);
-				RemoveAttribute(Pair.Key);
+				DuplicateAttribute(Pair.Key, NewIdentifier, bShouldTransact);
+				RemoveAttribute(Pair.Key, bShouldTransact);
 			}	
 		}
 	}
@@ -2348,7 +2360,7 @@ void UAnimSequencerController::UpdateWithSkeleton(USkeleton* TargetSkeleton, boo
 			// Notify of skeleton change
 			Model->GetNotifier().Notify(EAnimDataModelNotifyType::SkeletonChanged);
 		}
-		CloseBracket();
+		CloseBracket(bShouldTransact);
 	}	
 }
 
diff --git a/Engine/Plugins/Animation/AnimationData/Source/AnimationData/Private/AnimSequencerController.h b/Engine/Plugins/Animation/AnimationData/Source/AnimationData/Private/AnimSequencerController.h
--- a/Engine/Plugins/Animation/AnimationData/Source/AnimationData/Private/AnimSequencerController.h
+++ b/Engine/Plugins/Animation/AnimationData/Source/AnimationData/Private/AnimSequencerController.h
@@ -121,8 +121,10 @@ private:
 	void EnsureModelIsInitialized() const;
 
 	bool IgnoreSkeletonValidation() const;
+	virtual bool ShouldTransact(bool bInLocalShouldTransact) const override;
 private: 
 	int32 BracketDepth;
+	bool bShouldTransactBracket = false;
 
 	// UObject (typed) pointer to Model
 	UPROPERTY(transient)
diff --git a/Engine/Source/Developer/AnimationDataController/Private/AnimDataController.cpp b/Engine/Source/Developer/AnimationDataController/Private/AnimDataController.cpp
--- a/Engine/Source/Developer/AnimationDataController/Private/AnimDataController.cpp
+++ b/Engine/Source/Developer/AnimationDataController/Private/AnimDataController.cpp
@@ -39,8 +39,13 @@ void UAnimDataController::SetModel(TScriptInterface<IAnimationDataModel> InModel
 void UAnimDataController::OpenBracket(const FText& InTitle, bool bShouldTransact /*= true*/)
 {
 	ValidateModel();
+	
+	if (BracketDepth == 0)
+	{
+		bShouldTransactBracket = bShouldTransact;
+	}
 
-	if (UE::FChangeTransactor::CanTransactChanges() && !ChangeTransactor.IsTransactionPending() && bShouldTransact)
+	if (UE::FChangeTransactor::CanTransactChanges() && !ChangeTransactor.IsTransactionPending() && ShouldTransact(bShouldTransact))
 	{
 		ChangeTransactor.OpenTransaction(InTitle);
 
@@ -72,7 +77,7 @@ void UAnimDataController::CloseBracket(bool bShouldTransact /*= true*/)
 
 	if (BracketDepth == 0)
 	{
-		if (UE::FChangeTransactor::CanTransactChanges() && bShouldTransact)
+		if (UE::FChangeTransactor::CanTransactChanges() && ShouldTransact(bShouldTransact && bShouldTransactBracket))
 		{
 			ensure(ChangeTransactor.IsTransactionPending());
 
@@ -82,6 +87,8 @@ void UAnimDataController::CloseBracket(bool bShouldTransact /*= true*/)
 		}
 		
 		Model->GetNotifier().Notify(EAnimDataModelNotifyType::BracketClosed);
+
+		bShouldTransactBracket = false;
 	}
 }
 
@@ -291,7 +298,7 @@ bool UAnimDataController::RemoveBoneTracksMissingFromSkeleton(const USkeleton* S
 			FBracket Bracket = ConditionalBracket(LOCTEXT("RemoveBoneTracksMissingFromSkeleton", "Validating Bone Animation Track Data against Skeleton"), bShouldTransact);
 			for (const FName& TrackName : TracksToBeRemoved)
 			{
-				RemoveBoneTrack(TrackName);
+				RemoveBoneTrack(TrackName, bShouldTransact);
 			}
 
 			for (const FName& TrackName : TracksUpdated)
@@ -341,7 +348,7 @@ void UAnimDataController::UpdateAttributesFromSkeleton(const USkeleton* Skeleton
 			FBracket Bracket = ConditionalBracket(LOCTEXT("VerifyAttributeBoneNames", "Remapping Animation Attribute Data"), bShouldTransact);
 			for (const FAnimationAttributeIdentifier& Identifier : ToRemoveIdentifiers)
 			{
-				RemoveAttribute(Identifier);
+				RemoveAttribute(Identifier, bShouldTransact);
 			}
 			
 			for (const TPair<FAnimationAttributeIdentifier, int32>& Pair : ToDuplicateIdentifiers)
@@ -349,8 +356,8 @@ void UAnimDataController::UpdateAttributesFromSkeleton(const USkeleton* Skeleton
 				const FAnimationAttributeIdentifier& DuplicateIdentifier = Pair.Key;
 				FAnimationAttributeIdentifier NewIdentifier(DuplicateIdentifier.GetName(), Pair.Value, DuplicateIdentifier.GetBoneName(), DuplicateIdentifier.GetType());
 
-				DuplicateAttribute(Pair.Key, NewIdentifier);
-				RemoveAttribute(Pair.Key);
+				DuplicateAttribute(Pair.Key, NewIdentifier, bShouldTransact);
+				RemoveAttribute(Pair.Key, bShouldTransact);
 			}	
 		}
 	}
@@ -1710,6 +1717,11 @@ void UAnimDataController::ResizeAttributes(float NewLength, bool bInserted, floa
 	}
 }
 
+bool UAnimDataController::ShouldTransact(bool bInLocalShouldTransact) const
+{
+	return bInLocalShouldTransact && (BracketDepth == 0 || bShouldTransactBracket);
+}
+
 bool UAnimDataController::AddAttribute(const FAnimationAttributeIdentifier& AttributeIdentifier, bool bShouldTransact /*= true*/)
 {
 	if (AttributeIdentifier.IsValid())
@@ -2063,12 +2075,12 @@ void UAnimDataController::UpdateWithSkeleton(USkeleton* TargetSkeleton, bool bSh
 {
 	OpenBracket(LOCTEXT("SettingNewskeleton", "Updating Skeleton for Animation Data Model"), bShouldTransact);
 	{
-		RemoveBoneTracksMissingFromSkeleton(TargetSkeleton);
+		RemoveBoneTracksMissingFromSkeleton(TargetSkeleton, bShouldTransact);
 
 		// Notify of skeleton change
 		Model->GetNotifier().Notify(EAnimDataModelNotifyType::SkeletonChanged);
 	}
-	CloseBracket();
+	CloseBracket(bShouldTransact);
 }
 
 void UAnimDataController::PopulateWithExistingModel(TScriptInterface<IAnimationDataModel> InModel)
diff --git a/Engine/Source/Developer/AnimationDataController/Private/AnimDataControllerActions.cpp b/Engine/Source/Developer/AnimationDataController/Private/AnimDataControllerActions.cpp
--- a/Engine/Source/Developer/AnimationDataController/Private/AnimDataControllerActions.cpp
+++ b/Engine/Source/Developer/AnimationDataController/Private/AnimDataControllerActions.cpp
@@ -293,7 +293,7 @@ TUniquePtr<FChange> FAddFloatCurveAction::ExecuteInternal(IAnimationDataModel* M
 	Controller->AddCurve(CurveId, Flags, false);
 	Controller->SetCurveKeys(CurveId, Keys, false);
 	Controller->SetCurveColor(CurveId, Color, false);
-	Controller->CloseBracket();
+	Controller->CloseBracket(false);
 
 	return InverseAction;
 }
@@ -601,4 +601,4 @@ FString FSetAtributeKeysAction::ToStringInternal() const
 
 #endif // WITH_EDITOR
 
-#undef LOCTEXT_NAMESPACE // "AnimDataControllerActions"
\ No newline at end of file
+#undef LOCTEXT_NAMESPACE // "AnimDataControllerActions"
diff --git a/Engine/Source/Developer/AnimationDataController/Public/AnimDataController.h b/Engine/Source/Developer/AnimationDataController/Public/AnimDataController.h
--- a/Engine/Source/Developer/AnimationDataController/Public/AnimDataController.h
+++ b/Engine/Source/Developer/AnimationDataController/Public/AnimDataController.h
@@ -100,11 +100,14 @@ private:
 	/** Resizes the curve/attribute data stored on the model according to the provided new length and time at which to insert or remove time */
 	ANIMATIONDATACONTROLLER_API void ResizeCurves(float NewLength, bool bInserted, float T0, float T1, bool bShouldTransact = true);
 	ANIMATIONDATACONTROLLER_API void ResizeAttributes(float NewLength, bool bInserted, float T0, float T1, bool bShouldTransact = true);
+	
+	virtual bool ShouldTransact(bool bInLocalShouldTransact) const override;
 #endif // WITH_EDITOR
 
 private: 
 #if WITH_EDITOR
 	int32 BracketDepth;
+	bool bShouldTransactBracket = false;
 #endif // WITH_EDITOR
 
 #if WITH_EDITORONLY_DATA
diff --git a/Engine/Source/Runtime/Engine/Classes/Animation/AnimData/IAnimationDataController.h b/Engine/Source/Runtime/Engine/Classes/Animation/AnimData/IAnimationDataController.h
--- a/Engine/Source/Runtime/Engine/Classes/Animation/AnimData/IAnimationDataController.h
+++ b/Engine/Source/Runtime/Engine/Classes/Animation/AnimData/IAnimationDataController.h
@@ -717,6 +717,8 @@ protected:
 	/** Functionality used by FOpenBracketAction and FCloseBracketAction to broadcast their equivalent notifies without actually opening a bracket. */
 	virtual void NotifyBracketOpen() = 0;
 	virtual void NotifyBracketClosed() = 0;
+	
+	virtual bool ShouldTransact(bool bInLocalShouldTransact) const = 0;
 
 	template <typename FmtType, typename... Types>
     void ReportWarningf(const FmtType& Fmt, Types... Args) const
@@ -837,7 +839,7 @@ private:
 		typedef TUniquePtr<UE::FScopedCompoundTransaction> FTransaction; \
 		FTransaction ConditionalTransaction(const FText& Description, bool bCondition) \
 		{ \
-			if (UE::FChangeTransactor::CanTransactChanges() && bCondition) \
+			if (UE::FChangeTransactor::CanTransactChanges() && ShouldTransact(bCondition)) \
 			{  \
 				return MakeUnique<UE::FScopedCompoundTransaction>(ChangeTransactor, Description); \
 			} \
@@ -846,12 +848,12 @@ private:
 		typedef TUniquePtr<IAnimationDataController::FScopedBracket> FBracket;  \
 		FBracket ConditionalBracket(const FText& Description, bool bCondition) \
 		{	 \
-			return MakeUnique<IAnimationDataController::FScopedBracket>(this, Description, UE::FChangeTransactor::CanTransactChanges() && bCondition);	 \
+			return MakeUnique<IAnimationDataController::FScopedBracket>(this, Description, UE::FChangeTransactor::CanTransactChanges() && ShouldTransact(bCondition));	 \
 		} \
 		template<typename ActionClass, class... ActionArguments> \
 		void ConditionalAction(bool bCondition, ActionArguments&&... Arguments) \
 		{ \
-			if (UE::FChangeTransactor::CanTransactChanges() && bCondition) \
+			if (UE::FChangeTransactor::CanTransactChanges() && ShouldTransact(bCondition)) \
 			{  \
 				ChangeTransactor.AddTransactionChange<ActionClass>(Forward<ActionArguments>(Arguments)...); \
 			} \
diff --git a/Engine/Source/Runtime/Engine/Private/Animation/AnimSequenceBase.cpp b/Engine/Source/Runtime/Engine/Private/Animation/AnimSequenceBase.cpp
--- a/Engine/Source/Runtime/Engine/Private/Animation/AnimSequenceBase.cpp
+++ b/Engine/Source/Runtime/Engine/Private/Animation/AnimSequenceBase.cpp
@@ -186,7 +186,7 @@ void UAnimSequenceBase::PostLoad()
 		    		Controller->NotifyPopulated();
 		    	}
 		    	bPopulatingDataModel = false;
-		    	Controller->CloseBracket();
+		    	Controller->CloseBracket(bDoNotTransactAction);
 		    }
 		    else
 		    {
